First draft of OnlineServicesNull with changes to OnlineTests to support properly testing online services. DemoAuthTest runs the login pipeline which will construct a login and logout test.

diff --git a/Engine/Plugins/Online/Experimental/OnlineServices/Source/OnlineServicesCommon/Private/Online/OnlineSyntaxTests.cpp b/Engine/Plugins/Online/Experimental/OnlineServices/Source/OnlineServicesCommon/Private/Online/OnlineSyntaxTests.cpp
--- a/Engine/Plugins/Online/Experimental/OnlineServices/Source/OnlineServicesCommon/Private/Online/OnlineSyntaxTests.cpp
+++ b/Engine/Plugins/Online/Experimental/OnlineServices/Source/OnlineServicesCommon/Private/Online/OnlineSyntaxTests.cpp
@@ -443,7 +443,7 @@ void TestToLogStringSyntax()
 
 	ToLogString(FString());
 	ToLogString(UE::Online::FAuthLogin::Params());
-	ToLogString(UE::Online::FAuthLogin::Result());
+	//ToLogString(UE::Online::FAuthLogin::Result()); default constructing a TSharedPtr will crash
 	ToLogString(UE::Online::FAuthLogout::Params());
 	ToLogString(UE::Online::FAuthLogout::Result());
 	ToLogString(UE::Online::FAuthGenerateAuth::Params());
diff --git a/Engine/Plugins/Online/Experimental/OnlineServices/Source/OnlineServicesInterface/Public/Online/Auth.h b/Engine/Plugins/Online/Experimental/OnlineServices/Source/OnlineServicesInterface/Public/Online/Auth.h
--- a/Engine/Plugins/Online/Experimental/OnlineServices/Source/OnlineServicesInterface/Public/Online/Auth.h
+++ b/Engine/Plugins/Online/Experimental/OnlineServices/Source/OnlineServicesInterface/Public/Online/Auth.h
@@ -54,6 +54,9 @@ struct FAuthLogin
 	struct Result
 	{
 		TSharedRef<FAccountInfo> AccountInfo;
+
+	public:
+		Result() = delete; // cannot default construct due to TSharedRef
 	};
 };
 
diff --git a/Engine/Plugins/Online/Experimental/OnlineServicesNull/Source/OnlineServicesNull.Build.cs b/Engine/Plugins/Online/Experimental/OnlineServicesNull/Source/OnlineServicesNull.Build.cs
new file mode 100644
--- /dev/null
+++ b/Engine/Plugins/Online/Experimental/OnlineServicesNull/Source/OnlineServicesNull.Build.cs
@@ -0,0 +1,18 @@
+// Copyright Epic Games, Inc. All Rights Reserved.
+
+using UnrealBuildTool;
+using System.IO;
+
+public class OnlineServicesNull : ModuleRules
+{
+	public OnlineServicesNull(ReadOnlyTargetRules Target) : base(Target)
+    {
+		PublicDependencyModuleNames.AddRange(
+			new string[] {
+				"Core",
+				"OnlineServicesInterface",
+				"OnlineServicesCommon",
+			}
+		);
+	}
+}
diff --git a/Engine/Plugins/Online/Experimental/OnlineServicesNull/Source/Private/AuthNull.cpp b/Engine/Plugins/Online/Experimental/OnlineServicesNull/Source/Private/AuthNull.cpp
new file mode 100644
--- /dev/null
+++ b/Engine/Plugins/Online/Experimental/OnlineServicesNull/Source/Private/AuthNull.cpp
@@ -0,0 +1,275 @@
+// Copyright Epic Games, Inc. All Rights Reserved.
+
+#include "AuthNull.h"
+
+#include "OnlineServicesNull.h"
+#include "OnlineServicesNullTypes.h"
+#include "Online/AuthErrors.h"
+#include "Online/OnlineErrorDefinitions.h"
+
+namespace UE::Online {
+// Copied from OSS Null
+
+#define NULL_OSS_STRING_BUFFER_LENGTH 256
+#define NULL_MAX_TOKEN_SIZE 4096
+
+
+FAuthNull::FAuthNull(FOnlineServicesNull& InServices)
+	: FAuthCommon(InServices)
+{
+}
+
+void FAuthNull::Initialize()
+{
+	FAuthCommon::Initialize();
+
+}
+
+void FAuthNull::PreShutdown()
+{
+}
+
+FString FAuthNull::GenerateRandomUserId(int32 LocalUserNum)
+{
+	FString HostName;
+	/*if(ISocketSubsystem::Get(PLATFORM_SOCKETSUBSYSTEM))
+	{
+		if (!ISocketSubsystem::Get(PLATFORM_SOCKETSUBSYSTEM)->GetHostName(HostName))
+		{
+			// could not get hostname, use address
+			bool bCanBindAll;
+			TSharedPtr<class FInternetAddr> Addr = ISocketSubsystem::Get(PLATFORM_SOCKETSUBSYSTEM)->GetLocalHostAddr(*GLog, bCanBindAll);
+			HostName = Addr->ToString(false);
+		}
+	}*/
+
+	bool bUseStableNullId = false;
+	FString UserSuffix;
+
+	if (true)
+	{
+		UserSuffix = FString::Printf(TEXT("-%d"), LocalUserNum);
+	}
+	 
+#if !(UE_BUILD_SHIPPING && WITH_EDITOR)
+	if (GIsFirstInstance && !GIsEditor)
+	{
+		// If we're outside the editor and know this is the first instance, use the system login id
+		bUseStableNullId = true;
+	}
+#endif
+
+	if (bUseStableNullId)
+	{
+		// Use a stable id possibly with a user num suffix
+		return FString::Printf(TEXT("%s-%s%s"), *HostName, *FPlatformMisc::GetLoginId().ToUpper(), *UserSuffix);
+	}
+
+	// If we're not the first instance (or in the editor), return truly random id
+	return FString::Printf(TEXT("%s-%s%s"), *HostName, *FGuid::NewGuid().ToString(), *UserSuffix);
+}
+
+TOnlineAsyncOpHandle<FAuthLogin> FAuthNull::Login(FAuthLogin::Params&& Params)
+{
+	TOnlineAsyncOpRef<FAuthLogin> Op = GetOp<FAuthLogin>(MoveTemp(Params));
+	int32 LocalUserNum = FPlatformMisc::GetUserIndexForPlatformUser(Op->GetParams().PlatformUserId);
+
+	Op->Then([this, LocalUserNum](TOnlineAsyncOp<FAuthLogin>& InAsyncOp) mutable
+	{
+		TSharedRef<FAccountInfoNull> AccountInfo = MakeShared<FAccountInfoNull>();
+		FString DisplayId = GenerateRandomUserId(LocalUserNum);
+		AccountInfo->PlatformUserId = InAsyncOp.GetParams().PlatformUserId;
+		AccountInfo->DisplayName = DisplayId;
+		AccountInfo->UserId = FOnlineAccountIdRegistryNull::Get().Create(DisplayId, AccountInfo->PlatformUserId);
+		FAuthLogin::Result Result{ AccountInfo };
+		AccountInfos.Add(AccountInfo->UserId, AccountInfo);
+		InAsyncOp.SetResult(MoveTemp(Result));
+	})
+	.Enqueue(GetSerialQueue());
+
+	return Op->GetHandle();
+}
+
+TOnlineAsyncOpHandle<FAuthLogout> FAuthNull::Logout(FAuthLogout::Params&& Params)
+{
+	// todo: login status change delegates
+	TOnlineAsyncOpRef<FAuthLogout> Op = GetOp<FAuthLogout>(MoveTemp(Params));
+
+	Op->Then([this](TOnlineAsyncOp<FAuthLogout>& InAsyncOp) mutable
+	{
+		auto Result = GetAccountByAccountId({InAsyncOp.GetParams().LocalUserId});
+		if(Result.IsOk())
+		{
+			AccountInfos.Remove(Result.GetOkValue().AccountInfo->UserId);
+			InAsyncOp.SetResult(FAuthLogout::Result());
+		}
+		else
+		{
+			InAsyncOp.SetError(Errors::Unknown());
+		}
+	})
+	.Enqueue(GetSerialQueue());
+
+	return Op->GetHandle();
+}
+
+TOnlineAsyncOpHandle<FAuthGenerateAuth> FAuthNull::GenerateAuth(FAuthGenerateAuth::Params&& Params)
+{
+	TSharedRef<TOnlineAsyncOp<FAuthGenerateAuth>> AsyncOperation = MakeShared<TOnlineAsyncOp<FAuthGenerateAuth>>(Services, MoveTemp(Params));
+	return AsyncOperation->GetHandle();
+}
+
+TOnlineResult<FAuthGetAccountByAccountId> FAuthNull::GetAccountByAccountId(FAuthGetAccountByAccountId::Params&& Params)
+{
+	if (TSharedRef<FAccountInfoNull>* const FoundAccount = AccountInfos.Find(Params.LocalUserId))
+	{
+		return TOnlineResult<FAuthGetAccountByAccountId>({*FoundAccount});
+	}
+	else
+	{
+		// TODO: proper error
+		return TOnlineResult<FAuthGetAccountByAccountId>(Errors::Unknown());
+	}
+}
+
+bool FAuthNull::IsLoggedIn(const FOnlineAccountIdHandle& AccountId) const
+{
+	// TODO:  More logic?
+	return AccountInfos.Contains(AccountId);
+}
+
+TResult<FOnlineAccountIdHandle, FOnlineError> FAuthNull::GetAccountIdByLocalUserNum(int32 LocalUserNum) const
+{
+	FPlatformUserId UserId = FPlatformMisc::GetPlatformUserForUserIndex(LocalUserNum);
+	for (const TPair<FOnlineAccountIdHandle, TSharedRef<FAccountInfoNull>>& AccountPair : AccountInfos)
+	{
+		if (AccountPair.Value->PlatformUserId == UserId)
+		{
+			TResult<FOnlineAccountIdHandle, FOnlineError> Result(AccountPair.Key);
+			return Result;
+		}
+	}
+	TResult<FOnlineAccountIdHandle, FOnlineError> Result(Errors::Unknown()); // TODO: error code
+	return Result;
+}
+
+TOnlineResult<FAuthGetAccountByPlatformUserId> FAuthNull::GetAccountByPlatformUserId(FAuthGetAccountByPlatformUserId::Params&& Params)
+{
+	for (const TPair<FOnlineAccountIdHandle, TSharedRef<FAccountInfoNull>>& AccountPair : AccountInfos)
+	{
+		if (AccountPair.Value->PlatformUserId == Params.PlatformUserId)
+		{
+			return TOnlineResult<FAuthGetAccountByPlatformUserId>({AccountPair.Value});
+		}
+	}
+	TOnlineResult<FAuthGetAccountByPlatformUserId> Result(Errors::Unknown()); // TODO: error code
+	return Result;
+}
+
+// FOnlineAccountIdRegistryNull
+FOnlineAccountIdRegistryNull& FOnlineAccountIdRegistryNull::Get()
+{
+	static FOnlineAccountIdRegistryNull Instance;
+	return Instance;
+}
+
+FOnlineAccountIdHandle FOnlineAccountIdRegistryNull::Find(FString UserId) const
+{
+	const FOnlineAccountIdString* Entry = StringToId.Find(UserId);
+	if(Entry)
+	{
+		return Entry->Handle;
+	}
+	return FOnlineAccountIdHandle();
+}
+
+FOnlineAccountIdHandle FOnlineAccountIdRegistryNull::Find(FPlatformUserId UserId) const
+{
+	const FOnlineAccountIdString* Entry = LocalUserMap.Find(UserId.GetInternalId());
+	if (Entry)
+	{
+		return Entry->Handle;
+	}
+	return FOnlineAccountIdHandle();
+}
+
+FOnlineAccountIdHandle FOnlineAccountIdRegistryNull::Find(int32 UserId) const
+{
+	const FOnlineAccountIdString* Entry = LocalUserMap.Find(UserId);
+	if (Entry)
+	{
+		return Entry->Handle;
+	}
+	return FOnlineAccountIdHandle();
+}
+
+const FOnlineAccountIdString* FOnlineAccountIdRegistryNull::GetInternal(const FOnlineAccountIdHandle& Handle) const
+{
+	if(Handle.GetOnlineServicesType() == EOnlineServices::Null && Handle.GetHandle() <= (uint32)Ids.Num())
+	{
+		return &Ids[Handle.GetHandle()-1];
+	}
+	return nullptr;
+}
+
+FOnlineAccountIdHandle FOnlineAccountIdRegistryNull::Create(FString UserId, FPlatformUserId LocalUserIndex/* = FPlatformUserId::PLATFORMUSERID_NONE*/)
+{
+	FOnlineAccountIdHandle ExistingHandle = Find(UserId);
+	if(ExistingHandle.IsValid())
+	{
+		return ExistingHandle;
+	}
+
+	FOnlineAccountIdString Id = Ids.Emplace_GetRef();
+	Id.AccountIndex = Ids.Num();
+	Id.Data = UserId;
+	Id.Handle = FOnlineAccountIdHandle(EOnlineServices::Null, Id.AccountIndex);
+	
+	StringToId.Add(UserId, Id);
+	if(LocalUserIndex.IsValid())
+	{
+		if(LocalUserMap.Contains(LocalUserIndex.GetInternalId()))
+		{
+			UE_LOG(LogTemp, Warning, TEXT("OssNull: Found a duplicate ID for local user %d"), LocalUserIndex.GetInternalId());
+		}
+		LocalUserMap.Add(LocalUserIndex.GetInternalId(), Id);
+	}
+
+	return Id.Handle;
+}
+
+FString FOnlineAccountIdRegistryNull::ToLogString(const FOnlineAccountIdHandle& Handle) const
+{
+	if(const FOnlineAccountIdString* Id = GetInternal(Handle))
+	{
+		return Id->Data;
+	}
+
+	return FString(TEXT("[InvalidNetID]"));
+}
+
+
+TArray<uint8> FOnlineAccountIdRegistryNull::ToReplicationData(const FOnlineAccountIdHandle& Handle) const
+{
+	if (const FOnlineAccountIdString* Id = GetInternal(Handle))
+	{
+		TArray<uint8> ReplicationData;
+		ReplicationData.Reserve(Id->Data.Len());
+		StringToBytes(Id->Data, ReplicationData.GetData(), Id->Data.Len());
+		return ReplicationData;
+	}
+
+	return TArray<uint8>();;
+}
+
+FOnlineAccountIdHandle FOnlineAccountIdRegistryNull::FromReplicationData(const TArray<uint8> ReplicationData)
+{
+	FString Result = BytesToString(ReplicationData.GetData(), ReplicationData.Num());
+	if(Result.Len() > 0)
+	{
+		return Create(Result);
+	}
+	return FOnlineAccountIdHandle();
+}
+
+/* UE::Online */ }
diff --git a/Engine/Plugins/Online/Experimental/OnlineServicesNull/Source/Private/OnlineServicesNull.cpp b/Engine/Plugins/Online/Experimental/OnlineServicesNull/Source/Private/OnlineServicesNull.cpp
new file mode 100644
--- /dev/null
+++ b/Engine/Plugins/Online/Experimental/OnlineServicesNull/Source/Private/OnlineServicesNull.cpp
@@ -0,0 +1,48 @@
+// Copyright Epic Games, Inc. All Rights Reserved.
+
+#include "OnlineServicesNull.h"
+#include "OnlineServicesNullTypes.h"
+
+#include "AuthNull.h"
+
+namespace UE::Online {
+
+struct FNullPlatformConfig
+{
+public:
+	FString TestId;
+};
+
+namespace Meta {
+
+BEGIN_ONLINE_STRUCT_META(FNullPlatformConfig)
+// example config setup:
+//	ONLINE_STRUCT_FIELD(FNullPlatformConfig, TestId)
+END_ONLINE_STRUCT_META()
+
+ /*Meta*/ }
+
+FOnlineServicesNull::FOnlineServicesNull()
+	: FOnlineServicesCommon(TEXT("Null"))
+{
+}
+
+void FOnlineServicesNull::RegisterComponents()
+{
+	Components.Register<FAuthNull>(*this);
+	FOnlineServicesCommon::RegisterComponents();
+}
+
+void FOnlineServicesNull::Initialize()
+{
+	FNullPlatformConfig NullPlatformConfig;
+	LoadConfig(NullPlatformConfig);
+
+// example config loading:
+//	FTCHARToUTF8 TestId(*NullPlatformConfig.TestId);
+
+	FOnlineServicesCommon::Initialize();
+}
+
+
+/* UE::Online */ }
diff --git a/Engine/Plugins/Online/Experimental/OnlineServicesNull/Source/Private/OnlineServicesNullModule.cpp b/Engine/Plugins/Online/Experimental/OnlineServicesNull/Source/Private/OnlineServicesNullModule.cpp
new file mode 100644
--- /dev/null
+++ b/Engine/Plugins/Online/Experimental/OnlineServicesNull/Source/Private/OnlineServicesNullModule.cpp
@@ -0,0 +1,39 @@
+// Copyright Epic Games, Inc. All Rights Reserved.
+
+#include "CoreMinimal.h"
+#include "Modules/ModuleInterface.h"
+#include "Modules/ModuleManager.h"
+
+#include "Online/OnlineServicesRegistry.h"
+#include "OnlineServicesNull.h"
+
+namespace UE::Online
+{
+
+class FOnlineServicesNullModule : public IModuleInterface
+{
+public:
+	virtual void StartupModule() override;
+protected:
+};
+
+class FOnlineServicesFactoryNull : public IOnlineServicesFactory
+{
+public:
+	virtual ~FOnlineServicesFactoryNull() {}
+	virtual TSharedPtr<IOnlineServices> Create() override
+	{
+		return MakeShared<FOnlineServicesNull>();
+	}
+protected:
+};
+
+void FOnlineServicesNullModule::StartupModule()
+{
+	FOnlineServicesRegistry::Get().RegisterServicesFactory(EOnlineServices::Null, MakeUnique<FOnlineServicesFactoryNull>());
+	FOnlineIdRegistryRegistry::Get().RegisterAccountIdRegistry(EOnlineServices::Null, &FOnlineAccountIdRegistryNull::Get());
+}
+
+/* UE::Online */ }
+
+IMPLEMENT_MODULE(UE::Online::FOnlineServicesNullModule, OnlineServicesNull);
diff --git a/Engine/Plugins/Online/Experimental/OnlineServicesNull/Source/Public/AuthNull.h b/Engine/Plugins/Online/Experimental/OnlineServicesNull/Source/Public/AuthNull.h
new file mode 100644
--- /dev/null
+++ b/Engine/Plugins/Online/Experimental/OnlineServicesNull/Source/Public/AuthNull.h
@@ -0,0 +1,76 @@
+// Copyright Epic Games, Inc. All Rights Reserved.
+
+#pragma once
+
+#include "CoreMinimal.h"
+#include "Online/AuthCommon.h"
+
+namespace UE::Online {
+
+class FOnlineAccountIdString
+{
+public:
+	FString Data;
+	int32 AccountIndex;
+	FOnlineAccountIdHandle Handle;
+};
+
+class FOnlineAccountIdRegistryNull : public IOnlineAccountIdRegistry
+{
+public:
+	static FOnlineAccountIdRegistryNull& Get();
+
+	FOnlineAccountIdHandle Find(FString UserId) const;
+	FOnlineAccountIdHandle Find(FPlatformUserId UserId) const;
+	FOnlineAccountIdHandle Find(int32 UserId) const;
+
+	FOnlineAccountIdHandle Create(FString UserId, FPlatformUserId LocalUserIndex = PLATFORMUSERID_NONE);
+
+	// Begin IOnlineAccountIdRegistry
+	virtual FString ToLogString(const FOnlineAccountIdHandle& Handle) const override;
+	virtual TArray<uint8> ToReplicationData(const FOnlineAccountIdHandle& Handle) const override;
+	virtual FOnlineAccountIdHandle FromReplicationData(const TArray<uint8> ReplicationString) override;
+	// End IOnlineAccountIdRegistry
+
+	virtual ~FOnlineAccountIdRegistryNull() = default;
+
+private:
+	const FOnlineAccountIdString* GetInternal(const FOnlineAccountIdHandle& Handle) const;
+
+	// how much of these are actually necessary?
+	TArray<FOnlineAccountIdString> Ids;
+	TMap<FString, FOnlineAccountIdString> StringToId; 
+	TMap<int32, FOnlineAccountIdString> LocalUserMap;
+
+};
+
+class FOnlineServicesNull;
+
+class ONLINESERVICESNULL_API FAuthNull : public FAuthCommon
+{
+public:
+	using Super = FAuthCommon;
+
+	FAuthNull(FOnlineServicesNull& InOwningSubsystem);
+	virtual void Initialize() override;
+	virtual void PreShutdown() override;
+	virtual TOnlineAsyncOpHandle<FAuthLogin> Login(FAuthLogin::Params&& Params) override;
+	virtual TOnlineAsyncOpHandle<FAuthLogout> Logout(FAuthLogout::Params&& Params) override;
+	virtual TOnlineAsyncOpHandle<FAuthGenerateAuth> GenerateAuth(FAuthGenerateAuth::Params&& Params) override;
+	virtual TOnlineResult<FAuthGetAccountByPlatformUserId> GetAccountByPlatformUserId(FAuthGetAccountByPlatformUserId::Params&& Params) override;
+	virtual TOnlineResult<FAuthGetAccountByAccountId> GetAccountByAccountId(FAuthGetAccountByAccountId::Params&& Params) override;
+
+	bool IsLoggedIn(const FOnlineAccountIdHandle& AccountId) const;
+
+protected:
+	FString GenerateRandomUserId(int32 LocalUserNum);
+	TResult<FOnlineAccountIdHandle, FOnlineError> GetAccountIdByLocalUserNum(int32 LocalUserNum) const;
+
+	class FAccountInfoNull : public FAccountInfo
+	{
+	};
+
+	TMap<FOnlineAccountIdHandle, TSharedRef<FAccountInfoNull>> AccountInfos;
+};
+
+/* UE::Online */ }
diff --git a/Engine/Plugins/Online/Experimental/OnlineServicesNull/Source/Public/OnlineServicesNull.h b/Engine/Plugins/Online/Experimental/OnlineServicesNull/Source/Public/OnlineServicesNull.h
new file mode 100644
--- /dev/null
+++ b/Engine/Plugins/Online/Experimental/OnlineServicesNull/Source/Public/OnlineServicesNull.h
@@ -0,0 +1,21 @@
+// Copyright Epic Games, Inc. All Rights Reserved.
+
+#pragma once
+
+#include "Online/OnlineServicesCommon.h"
+
+namespace UE::Online {
+
+class ONLINESERVICESNULL_API FOnlineServicesNull : public FOnlineServicesCommon
+{
+public:
+	using Super = FOnlineServicesCommon;
+
+	FOnlineServicesNull();
+	virtual void RegisterComponents() override;
+	virtual void Initialize() override;
+
+	static EOnlineServices GetServicesProvider() { return EOnlineServices::Null; }
+};
+
+/* UE::Online */ }
diff --git a/Engine/Plugins/Online/Experimental/OnlineServicesNull/Source/Public/OnlineServicesNullTypes.h b/Engine/Plugins/Online/Experimental/OnlineServicesNull/Source/Public/OnlineServicesNullTypes.h
new file mode 100644
--- /dev/null
+++ b/Engine/Plugins/Online/Experimental/OnlineServicesNull/Source/Public/OnlineServicesNullTypes.h
@@ -0,0 +1,3 @@
+// Copyright Epic Games, Inc. All Rights Reserved.
+
+#pragma once
\ No newline at end of file
