Implementation of the Lobby interface for OSS Null v2.
This includes the main functionality required to make most basic Lyra LAN flows work. Biggest flaw currently is that we dont have a way to get the broadcasting port when advertising our lobby, so I have temporarily hardcoded 7777 as the broadcasting port. Having two hosts on the same PC and trying to connect to the second one with a third client will not work. The member list also does not reflect the true member list of the session which will be obtained once the game state is replicated to you.

Submitting for [REDACTED_HANDLE]

diff --git a/Engine/Plugins/Online/Experimental/OnlineServices/Source/OnlineServicesInterface/Private/Online/Lobbies.cpp b/Engine/Plugins/Online/Experimental/OnlineServices/Source/OnlineServicesInterface/Private/Online/Lobbies.cpp
--- a/Engine/Plugins/Online/Experimental/OnlineServices/Source/OnlineServicesInterface/Private/Online/Lobbies.cpp
+++ b/Engine/Plugins/Online/Experimental/OnlineServices/Source/OnlineServicesInterface/Private/Online/Lobbies.cpp
@@ -90,6 +90,10 @@ FString FLobbyVariant::GetString() const
 	{
 		return ::LexToString(VariantData.Get<bool>());
 	}
+	else if (VariantData.IsType<double>())
+	{
+		return FString::Printf(TEXT("%f"), VariantData.Get<double>());
+	}
 	else
 	{
 		checkNoEntry();
diff --git a/Engine/Plugins/Online/Experimental/OnlineServicesNull/OnlineServicesNull.uplugin b/Engine/Plugins/Online/Experimental/OnlineServicesNull/OnlineServicesNull.uplugin
new file mode 100644
--- /dev/null
+++ b/Engine/Plugins/Online/Experimental/OnlineServicesNull/OnlineServicesNull.uplugin
@@ -0,0 +1,26 @@
+{
+	"FileVersion": 3,
+
+	"FriendlyName": "Online Services Null",
+	"Version": 1,
+	"VersionName": "1.0",
+	"Description": "Online Services implementation without an external service.",
+	"Category": "Online Platform",
+	"CreatedBy": "Epic Games, Inc.",
+	"CreatedByURL": "http://epicgames.com",
+	"EnabledByDefault": false,
+
+	"Modules": [
+		{
+			"Name": "OnlineServicesNull",
+			"Type": "Runtime",
+			"LoadingPhase": "PostConfigInit"
+		}
+	],
+	"Plugins": [
+		{
+			"Name": "OnlineServices",
+			"Enabled": true
+		}
+	]
+}
\ No newline at end of file
diff --git a/Engine/Plugins/Online/Experimental/OnlineServicesNull/Source/OnlineServicesNull.Build.cs b/Engine/Plugins/Online/Experimental/OnlineServicesNull/Source/OnlineServicesNull.Build.cs
--- a/Engine/Plugins/Online/Experimental/OnlineServicesNull/Source/OnlineServicesNull.Build.cs
+++ b/Engine/Plugins/Online/Experimental/OnlineServicesNull/Source/OnlineServicesNull.Build.cs
@@ -10,8 +10,12 @@ public class OnlineServicesNull : ModuleRules
 		PublicDependencyModuleNames.AddRange(
 			new string[] {
 				"Core",
+				"CoreOnline",
+				"Sockets",
 				"OnlineServicesInterface",
 				"OnlineServicesCommon",
+				"OnlineSubsystem",
+				"OnlineSubsystemUtils"
 			}
 		);
 	}
diff --git a/Engine/Plugins/Online/Experimental/OnlineServicesNull/Source/Private/Online/AuthNull.cpp b/Engine/Plugins/Online/Experimental/OnlineServicesNull/Source/Private/Online/AuthNull.cpp
--- a/Engine/Plugins/Online/Experimental/OnlineServicesNull/Source/Private/Online/AuthNull.cpp
+++ b/Engine/Plugins/Online/Experimental/OnlineServicesNull/Source/Private/Online/AuthNull.cpp
@@ -62,11 +62,11 @@ FString FAuthNull::GenerateRandomUserId(int32 LocalUserNum)
 	if (bUseStableNullId)
 	{
 		// Use a stable id possibly with a user num suffix
-		return FString::Printf(TEXT("%s-%s%s"), *HostName, *FPlatformMisc::GetLoginId().ToUpper(), *UserSuffix);
+		return FString::Printf(TEXT("OSSV2-%s-%s%s"), *HostName, *FPlatformMisc::GetLoginId().ToUpper(), *UserSuffix);
 	}
 
 	// If we're not the first instance (or in the editor), return truly random id
-	return FString::Printf(TEXT("%s-%s%s"), *HostName, *FGuid::NewGuid().ToString(), *UserSuffix);
+	return FString::Printf(TEXT("OSSV2-%s-%s%s"), *HostName, *FGuid::NewGuid().ToString(), *UserSuffix);
 }
 
 TOnlineAsyncOpHandle<FAuthLogin> FAuthNull::Login(FAuthLogin::Params&& Params)
@@ -81,6 +81,7 @@ TOnlineAsyncOpHandle<FAuthLogin> FAuthNull::Login(FAuthLogin::Params&& Params)
 		AccountInfo->PlatformUserId = InAsyncOp.GetParams().PlatformUserId;
 		AccountInfo->DisplayName = DisplayId;
 		AccountInfo->UserId = FOnlineAccountIdRegistryNull::Get().Create(DisplayId, AccountInfo->PlatformUserId);
+		AccountInfo->LoginStatus = ELoginStatus::LoggedIn;
 		FAuthLogin::Result Result{ AccountInfo };
 		AccountInfos.Add(AccountInfo->UserId, AccountInfo);
 		InAsyncOp.SetResult(MoveTemp(Result));
@@ -220,7 +221,7 @@ FOnlineAccountIdHandle FOnlineAccountIdRegistryNull::Create(FString UserId, FPla
 		return ExistingHandle;
 	}
 
-	FOnlineAccountIdString Id = Ids.Emplace_GetRef();
+	FOnlineAccountIdString& Id = Ids.Emplace_GetRef();
 	Id.AccountIndex = Ids.Num();
 	Id.Data = UserId;
 	Id.Handle = FOnlineAccountIdHandle(EOnlineServices::Null, Id.AccountIndex);
@@ -230,7 +231,7 @@ FOnlineAccountIdHandle FOnlineAccountIdRegistryNull::Create(FString UserId, FPla
 	{
 		if(LocalUserMap.Contains(LocalUserIndex.GetInternalId()))
 		{
-			UE_LOG(LogTemp, Warning, TEXT("OssNull: Found a duplicate ID for local user %d"), LocalUserIndex.GetInternalId());
+			UE_LOG(LogTemp, Error, TEXT("OssNull: Found a duplicate ID for local user %d"), LocalUserIndex.GetInternalId());
 		}
 		LocalUserMap.Add(LocalUserIndex.GetInternalId(), Id);
 	}
@@ -254,15 +255,16 @@ TArray<uint8> FOnlineAccountIdRegistryNull::ToReplicationData(const FOnlineAccou
 	if (const FOnlineAccountIdString* Id = GetInternal(Handle))
 	{
 		TArray<uint8> ReplicationData;
-		ReplicationData.Reserve(Id->Data.Len());
+		ReplicationData.SetNumUninitialized(Id->Data.Len());
 		StringToBytes(Id->Data, ReplicationData.GetData(), Id->Data.Len());
+		UE_LOG(LogTemp, VeryVerbose, TEXT("StringToBytes on %s returned %d len"), *Id->Data, ReplicationData.Num())
 		return ReplicationData;
 	}
 
 	return TArray<uint8>();;
 }
 
-FOnlineAccountIdHandle FOnlineAccountIdRegistryNull::FromReplicationData(const TArray<uint8> ReplicationData)
+FOnlineAccountIdHandle FOnlineAccountIdRegistryNull::FromReplicationData(const TArray<uint8>& ReplicationData)
 {
 	FString Result = BytesToString(ReplicationData.GetData(), ReplicationData.Num());
 	if(Result.Len() > 0)
diff --git a/Engine/Plugins/Online/Experimental/OnlineServicesNull/Source/Private/Online/LobbiesNull.cpp b/Engine/Plugins/Online/Experimental/OnlineServicesNull/Source/Private/Online/LobbiesNull.cpp
new file mode 100644
--- /dev/null
+++ b/Engine/Plugins/Online/Experimental/OnlineServicesNull/Source/Private/Online/LobbiesNull.cpp
@@ -0,0 +1,405 @@
+// Copyright Epic Games, Inc. All Rights Reserved.
+
+#include "Online/LobbiesNull.h"
+#include "Misc/Guid.h"
+#include "SocketSubsystem.h"
+#include "Online/NboSerializerNullSvc.h"
+#include "UObject/CoreNet.h"
+#include "Online/LobbyRegistryNull.h"
+#include "OnlineSubsystemTypes.h" //todo: preferably shouldnt need this
+#include "OnlineSubsystemUtils.h" //todo: preferably shouldnt need this
+
+namespace UE::Online { 
+
+FLobbyNull::FLobbyNull()
+	: Data(MakeShared<FLobby>())
+{
+	bool bCanBindAll;
+	HostAddr = ISocketSubsystem::Get(PLATFORM_SOCKETSUBSYSTEM)->GetLocalHostAddr(*GLog, bCanBindAll);
+
+	// The below is a workaround for systems that set hostname to a distinct address from 127.0.0.1 on a loopback interface.
+	// See e.g. https://www.debian.org/doc/manuals/debian-reference/ch05.en.html#_the_hostname_resolution
+	// and http://serverfault.com/questions/363095/why-does-my-hostname-appear-with-the-address-127-0-1-1-rather-than-127-0-0-1-in
+	// Since we bind to 0.0.0.0, we won't answer on 127.0.1.1, so we need to advertise ourselves as 127.0.0.1 for any other loopback address we may have.
+	uint32 HostIp = 0;
+	HostAddr->GetIp(HostIp); // will return in host order
+	// if this address is on loopback interface, advertise it as 127.0.0.1
+	if ((HostIp & 0xff000000) == 0x7f000000)
+	{
+		HostAddr->SetIp(0x7f000001);	// 127.0.0.1
+	}
+
+	// Now set the port that was configured
+#if WITH_ENGINE
+	// jntodo: can we put anything here if we don't have engine code?
+	HostAddr->SetPort(GetPortFromNetDriver(NAME_None));
+#endif //WITH_ENGINE
+	
+	if (HostAddr->GetPort() == 0)
+	{
+		HostAddr->SetPort(7777); // temp workaround
+	}
+
+	HostAddr->GetIp(HostAddrIp);
+	HostAddr->GetPort(HostAddrPort);
+
+}
+
+void FLobbiesNull::Initialize()
+{
+
+}
+
+void FLobbiesNull::PreShutdown()
+{
+
+}
+
+void FLobbiesNull::Tick(float DeltaSeconds)
+{
+	LANSessionManager.Tick(DeltaSeconds);
+}
+
+TOnlineAsyncOpHandle<FCreateLobby> FLobbiesNull::CreateLobby(FCreateLobby::Params&& Params)
+{
+	TOnlineAsyncOpRef<FCreateLobby> Op = GetOp<FCreateLobby>(MoveTemp(Params));
+	TSharedRef<FLobbyNull> Lobby = CreateNamedLobby(Op->GetParams());
+	
+	int Result = UpdateLANStatus();
+
+	if (Result != ONLINE_IO_PENDING)
+	{
+		// Set the game state as pending (not started)
+//		Session->SessionState = EOnlineSessionState::Pending; todo: LobbyState
+
+		if (Result != ONLINE_SUCCESS)
+		{
+			// Clean up the session info so we don't get into a confused state
+			RemoveLobbyFromRef(Lobby);
+			Op->SetError(Errors::Unknown()); //todo: error
+		}
+		else
+		{
+			//RegisterLocalPlayers(Lobby); todo: RegisterPlayers
+			Op->SetResult(FCreateLobby::Result{Lobby->Data});
+		}
+	}
+	return Op->GetHandle();
+}
+
+TOnlineAsyncOpHandle<FFindLobbies> FLobbiesNull::FindLobbies(FFindLobbies::Params&& Params)
+{
+	UE_LOG(LogTemp, Warning, TEXT("FLobbiesNull::FindLobbies"));
+	if (!CurrentLobbySearch.IsValid())
+	{
+		CurrentLobbySearch = MakeShared<FFindLobbies::Result>();
+		CurrentLobbySearchHandle = GetOp<FFindLobbies>(MoveTemp(Params));
+
+		FindLANSession();
+	}
+
+	return CurrentLobbySearchHandle->GetHandle(); // todo: errors (if current lobby search is valid)
+}
+
+TOnlineAsyncOpHandle<FJoinLobby> FLobbiesNull::JoinLobby(FJoinLobby::Params&& Params)
+{
+	TOnlineAsyncOpRef<FJoinLobby> Op = GetOp<FJoinLobby>(MoveTemp(Params));
+	const FName& LobbyName = Op->GetParams().LocalName;
+	TSharedPtr<FLobbyNull> Lobby = GetLobby(Op->GetParams().LobbyId);
+	TSharedPtr<FLobbyNull> ExistingLobby = GetNamedLobby(LobbyName);
+
+	if (!ExistingLobby.IsValid() && Lobby.IsValid())
+	{
+		AddNamedLobby(Lobby.ToSharedRef(), LobbyName);
+	}
+	else
+	{
+		if (!Lobby.IsValid())
+		{
+			UE_LOG(LogTemp, Warning, TEXT("Failed to find a lobby with ID %d"), Op->GetParams().LobbyId.GetHandle())
+		}
+		else
+		{
+			UE_LOG(LogTemp, Warning, TEXT("Lobby with name %s already exists"), *LobbyName.ToString());
+		}
+		Op->SetError(Errors::RequestFailure());  
+		return Op->GetHandle();
+	}
+
+	Op->SetResult(FJoinLobby::Result{Lobby->Data});
+	return Op->GetHandle();
+
+}
+
+TOnlineAsyncOpHandle<FLeaveLobby> FLobbiesNull::LeaveLobby(FLeaveLobby::Params&& Params)
+{
+	TOnlineAsyncOpRef<FLeaveLobby> Op = GetOp<FLeaveLobby>(MoveTemp(Params));
+	RemoveLobbyFromId(Op->GetParams().LobbyId);
+	Op->SetResult(FLeaveLobby::Result{});
+	return Op->GetHandle();
+}
+
+TSharedRef<FLobbyNull> FLobbiesNull::CreateNamedLobby(const FCreateLobby::Params& Params)
+{
+	//CHECK(!NamedLobbies.Contains(Params.LocalName))
+	TSharedRef<FLobbyNull> Lobby = NamedLobbies.Add(Params.LocalName, MakeShared<FLobbyNull>());
+	Lobby->Data->LobbyId = FOnlineLobbyIdRegistryNull::Get().GetNext();
+	Lobby->Data->OwnerAccountId = Params.LocalUserId;
+	Lobby->Data->LocalName = Params.LocalName;
+	//Lobby.Data.SchemaName = ???
+	Lobby->Data->JoinPolicy = ELobbyJoinPolicy::PublicAdvertised; // temp
+	Lobby->Data->Attributes = Params.Attributes;
+	Lobby->Data->MaxMembers = 100; // temp
+
+	// construct the members
+	for (const FJoinLobbyLocalUserData& Data : Params.LocalUsers)
+	{
+		TSharedRef<FLobbyMember> NewMember = MakeShared<FLobbyMember>();
+
+		NewMember->AccountId = Data.LocalUserId;
+		NewMember->PlatformAccountId = Data.LocalUserId;
+		NewMember->PlatformDisplayName = TEXT("TEMP");
+		NewMember->Attributes = Data.Attributes;
+
+		Lobby->Data->Members.Add(Data.LocalUserId, NewMember);
+
+	}
+
+	AllLobbies.Add(Lobby->Data->LobbyId, Lobby);
+	return Lobby;
+}
+
+uint32 FLobbiesNull::FindLANSession()
+{
+	UE_LOG(LogTemp, Warning, TEXT("FLobbiesNull::FindLANSession"));
+	uint32 Return = ONLINE_IO_PENDING;
+
+	// Recreate the unique identifier for this client
+	GenerateNonce((uint8*)&LANSessionManager.LanNonce, 8);
+
+	FOnValidResponsePacketDelegate ResponseDelegate = FOnValidResponsePacketDelegate::CreateRaw(this, &FLobbiesNull::OnValidResponsePacketReceived);
+	FOnSearchingTimeoutDelegate TimeoutDelegate = FOnSearchingTimeoutDelegate::CreateRaw(this, &FLobbiesNull::OnLANSearchTimeout);
+
+	FNboSerializeToBufferNullSvc Packet(LAN_BEACON_MAX_PACKET_SIZE);
+	LANSessionManager.CreateClientQueryPacket(Packet, LANSessionManager.LanNonce);
+	if (LANSessionManager.Search(Packet, ResponseDelegate, TimeoutDelegate) == false)
+	{
+		UE_LOG(LogTemp, Warning, TEXT("FLobbiesNull::FindLANSession: Failed, returning"));
+		Return = ONLINE_FAIL;
+
+		FinalizeLANSearch();
+
+		// Just trigger the delegate as having failed
+		CurrentLobbySearchHandle->SetError(Errors::Unknown());
+		CurrentLobbySearch = nullptr;
+		CurrentLobbySearchHandle = nullptr;
+	}
+	else
+	{
+		UE_LOG(LogTemp, Warning, TEXT("FLobbiesNull::FindLANSession: searching...."));
+	}
+
+
+	return Return;
+}
+
+bool FLobbiesNull::NeedsToAdvertise()
+{
+	// todo: only advertise if we are the host
+	return !NamedLobbies.IsEmpty();
+}
+
+uint32 FLobbiesNull::UpdateLANStatus()
+{
+	UE_LOG(LogTemp, Warning, TEXT("FLobbiesNull::UpdateLANStatus"));
+	uint32 Result = ONLINE_SUCCESS;
+
+	if ( NeedsToAdvertise() )
+	{
+		// set up LAN session
+		if (LANSessionManager.GetBeaconState() == ELanBeaconState::NotUsingLanBeacon)
+		{
+			UE_LOG(LogTemp, Warning, TEXT("FLobbiesNull::UpdateLANStatus- Hosting.."));
+			FOnValidQueryPacketDelegate QueryPacketDelegate = FOnValidQueryPacketDelegate::CreateRaw(this, &FLobbiesNull::OnValidQueryPacketReceived);
+			if (!LANSessionManager.Host(QueryPacketDelegate))
+			{
+				UE_LOG(LogTemp, Warning, TEXT("FLobbiesNull::Failed to host!"));
+				Result = ONLINE_FAIL;
+
+				LANSessionManager.StopLANSession();
+			}
+		}
+	}
+	else
+	{
+		if (LANSessionManager.GetBeaconState() != ELanBeaconState::Searching)
+		{
+			// Tear down the LAN beacon
+			LANSessionManager.StopLANSession();
+		}
+	}
+
+	return Result;
+}
+
+void FLobbiesNull::AppendLobbyToPacket(FNboSerializeToBufferNullSvc& Packet, const TSharedRef<FLobbyNull>& Lobby)
+{
+	/** Owner of the session */
+	UE_LOG(LogTemp, Warning, TEXT("AppendLobbyToPacket: Appending IP %s"), *Lobby->HostAddr->ToString(true));
+
+	Packet << Lobby->Data->OwnerAccountId
+		<< Lobby->Data->Attributes
+		<< TEXT("Temp:Owner") // todo: owner name
+		<< Lobby->HostAddrIp
+		<< Lobby->HostAddrPort
+		<< 100 // todo: num private connections
+		<< 100; // todo: num public connections
+
+	// todo: host info, per game settings
+}
+
+void FLobbiesNull::OnValidQueryPacketReceived(uint8* PacketData, int32 PacketLength, uint64 ClientNonce)
+{
+	UE_LOG(LogTemp, Warning, TEXT("FLobbiesNull::OnValidQueryPacketReceived"));
+	// Iterate through all registered sessions and respond for each one that can be joinable
+	for (TPair<FName, TSharedRef<FLobbyNull>>& Pair : NamedLobbies)
+	{
+		const TSharedRef<FLobbyNull>& Lobby = Pair.Value;
+
+		// TODO: Joinability (ensure lobby is joinable
+
+		FNboSerializeToBufferNullSvc Packet(LAN_BEACON_MAX_PACKET_SIZE);
+		// Create the basic header before appending additional information
+		LANSessionManager.CreateHostResponsePacket(Packet, ClientNonce);
+
+		// Add all the session details
+		AppendLobbyToPacket(Packet, Lobby);
+
+		// Broadcast this response so the client can see us
+		if (!Packet.HasOverflow())
+		{
+			LANSessionManager.BroadcastPacket(Packet, Packet.GetByteCount());
+		}
+		else
+		{
+			UE_LOG(LogTemp, Warning, TEXT("LAN broadcast packet overflow, cannot broadcast on LAN"));
+		}
+	}
+}
+
+void FLobbiesNull::ReadLobbyFromPacket(FNboSerializeFromBufferNullSvc& Packet, const TSharedRef<FLobbyNull>& Session)
+{
+	UE_LOG(LogTemp, Verbose, TEXT("Reading session information from server"));
+
+	/** Owner of the session */
+	FOnlineAccountIdHandle OwningUserId;
+	FString OwningUserName;
+	uint32 NumOpenPrivateConnections, NumOpenPublicConnections;
+	Packet >> OwningUserId
+		>> Session->Data->Attributes
+		>> OwningUserName
+		>> Session->HostAddrIp
+		>> Session->HostAddrPort
+		>> NumOpenPrivateConnections
+		>> NumOpenPublicConnections;
+
+	Session->Data->OwnerAccountId = OwningUserId;
+	Session->HostAddr = ISocketSubsystem::Get(PLATFORM_SOCKETSUBSYSTEM)->CreateInternetAddr();
+	Session->HostAddr->SetIp(Session->HostAddrIp);
+	Session->HostAddr->SetPort(Session->HostAddrPort);
+	Session->Data->Attributes.Add(FName(TEXT("ConnectAddress")), Session->HostAddr->ToString(true));
+	Session->Data->LobbyId = FOnlineLobbyIdRegistryNull::Get().GetNext(); //TODO: this should be coming from the packet, not a random GUID
+	UE_LOG(LogTemp, Warning, TEXT("ReadLobby: Got IP %s"), *Session->HostAddr->ToString(true));
+
+	// todo: make all the members (not just the host
+	TSharedRef<FLobbyMember> NewMember = MakeShared<FLobbyMember>();
+	NewMember->AccountId = OwningUserId;
+	NewMember->PlatformAccountId = OwningUserId;
+	NewMember->PlatformDisplayName = TEXT("TEMP");
+	Session->Data->Members.Add(OwningUserId, NewMember);
+	AllLobbies.Add(Session->Data->LobbyId, Session);
+
+	// todo: host info, per game settings
+}
+
+void FLobbiesNull::OnValidResponsePacketReceived(uint8* PacketData, int32 PacketLength)
+{
+	UE_LOG(LogTemp, Warning, TEXT("FLobbiesNull::OnValidResponsePacketReceived"));
+	if (CurrentLobbySearch.IsValid())
+	{
+		TSharedRef<FLobbyNull> Lobby = MakeShared<FLobbyNull>();
+
+		// todo: ping
+		// this is not a correct ping, but better than nothing
+		//Lobby->PingInMs = static_cast<int32>((FPlatformTime::Seconds() - SessionSearchStartInSeconds) * 1000);
+
+		// Prepare to read data from the packet
+		FNboSerializeFromBufferNullSvc Packet(PacketData, PacketLength);
+		
+		ReadLobbyFromPacket(Packet, Lobby);
+
+		CurrentLobbySearch->Lobbies.Add(Lobby->Data);
+
+		// NOTE: we don't notify until the timeout happens
+	}
+	else
+	{
+		UE_LOG(LogTemp, Warning, TEXT("Failed to create new online game settings object"));
+	}
+}
+
+uint32 FLobbiesNull::FinalizeLANSearch()
+{
+	UE_LOG(LogTemp, Warning, TEXT("FLobbiesNull::FinalizeLANSearch"));
+	if (LANSessionManager.GetBeaconState() == ELanBeaconState::Searching)
+	{
+		LANSessionManager.StopLANSession();
+	}
+
+	return UpdateLANStatus();
+}
+
+void FLobbiesNull::OnLANSearchTimeout()
+{
+	UE_LOG(LogTemp, Warning, TEXT("FLobbiesNull::OnLANSearchTimeout"));
+	FinalizeLANSearch();
+
+	if (CurrentLobbySearch.IsValid())
+	{
+		if (CurrentLobbySearch->Lobbies.Num() > 0)
+		{
+			UE_LOG(LogTemp, Warning, TEXT("FLobbiesNull::OnLANSearchTimeout: Had results!"));
+			// Allow game code to sort the servers (todo)
+			//CurrentSessionSearch->SortSearchResults();
+			FFindLobbies::Result Result = *CurrentLobbySearch; // copy it- not sure if moving something owned by a shared ptr is a good idea
+			CurrentLobbySearchHandle->SetResult(MoveTemp(Result)); // todo: will this fire the delegates properly even though we reset the shared ptr?
+			
+			CurrentLobbySearch.Reset();
+			CurrentLobbySearchHandle.Reset();
+			return;
+		}
+	}
+	
+	UE_LOG(LogTemp, Warning, TEXT("FLobbiesNull::OnLANSearchTimeout: no results"));
+	// Trigger the delegate as complete
+	CurrentLobbySearchHandle->SetError(Errors::Unknown());
+	CurrentLobbySearch.Reset();
+	CurrentLobbySearchHandle.Reset();
+
+}
+
+TOnlineResult<FGetJoinedLobbies> FLobbiesNull::GetJoinedLobbies(FGetJoinedLobbies::Params&& Params)
+{
+	FGetJoinedLobbies::Result Result;
+	for (TPair<FName, TSharedRef<FLobbyNull>>& Pair : NamedLobbies)
+	{
+		Result.Lobbies.Add(Pair.Value->Data);
+	}
+	return TOnlineResult<FGetJoinedLobbies>(MoveTemp(Result));
+}
+
+
+} // namespace UE::Online
+
+
+#undef NOT_IMPLEMENTED
diff --git a/Engine/Plugins/Online/Experimental/OnlineServicesNull/Source/Private/Online/LobbyRegistryNull.cpp b/Engine/Plugins/Online/Experimental/OnlineServicesNull/Source/Private/Online/LobbyRegistryNull.cpp
new file mode 100644
--- /dev/null
+++ b/Engine/Plugins/Online/Experimental/OnlineServicesNull/Source/Private/Online/LobbyRegistryNull.cpp
@@ -0,0 +1,81 @@
+// Copyright Epic Games, Inc. All Rights Reserved.
+#include "Online/LobbyRegistryNull.h"
+
+namespace UE::Online {
+
+FOnlineLobbyIdRegistryNull& FOnlineLobbyIdRegistryNull::Get()
+{
+	static FOnlineLobbyIdRegistryNull Instance;
+	return Instance;
+}
+
+const FOnlineLobbyIdHandle* FOnlineLobbyIdRegistryNull::Find(FString LobbyId)
+{
+	return StringToId.Find(LobbyId);
+}
+
+FOnlineLobbyIdHandle FOnlineLobbyIdRegistryNull::FindOrAdd(FString LobbyId)
+{
+	const FOnlineLobbyIdHandle* Entry = StringToId.Find(LobbyId);
+	if (Entry)
+	{
+		return *Entry;
+	}
+
+	Ids.Add(LobbyId);
+	FOnlineLobbyIdHandle Handle(EOnlineServices::Null, Ids.Num());
+	StringToId.Add(LobbyId, Handle);
+
+	return Handle;
+}
+
+UE::Online::FOnlineLobbyIdHandle FOnlineLobbyIdRegistryNull::GetNext()
+{
+	return FindOrAdd(FGuid().ToString());
+}
+
+
+const FString* FOnlineLobbyIdRegistryNull::GetInternal(const FOnlineLobbyIdHandle& Handle) const
+{
+	if (Handle.GetOnlineServicesType() == EOnlineServices::Null && Handle.GetHandle() <= (uint32)Ids.Num())
+	{
+		return &Ids[Handle.GetHandle()-1];
+	}
+	return nullptr;
+}
+
+FString FOnlineLobbyIdRegistryNull::ToLogString(const FOnlineLobbyIdHandle& Handle) const
+{
+	if (const FString* Id = GetInternal(Handle))
+	{
+		return *Id;
+	}
+
+	return FString(TEXT("[InvalidLobbyID]"));
+}
+
+
+TArray<uint8> FOnlineLobbyIdRegistryNull::ToReplicationData(const FOnlineLobbyIdHandle& Handle) const
+{
+	if (const FString* Id = GetInternal(Handle))
+	{
+		TArray<uint8> ReplicationData;
+		ReplicationData.Reserve(Id->Len());
+		StringToBytes(*Id, ReplicationData.GetData(), Id->Len());
+		return ReplicationData;
+	}
+
+	return TArray<uint8>();;
+}
+
+FOnlineLobbyIdHandle FOnlineLobbyIdRegistryNull::FromReplicationData(const TArray<uint8>& ReplicationData)
+{
+	FString Result = BytesToString(ReplicationData.GetData(), ReplicationData.Num());
+	if (Result.Len() > 0)
+	{
+		return FindOrAdd(Result);
+	}
+	return FOnlineLobbyIdHandle();
+}
+
+} // namespace UE::Online
diff --git a/Engine/Plugins/Online/Experimental/OnlineServicesNull/Source/Private/Online/OnlineServicesNull.cpp b/Engine/Plugins/Online/Experimental/OnlineServicesNull/Source/Private/Online/OnlineServicesNull.cpp
--- a/Engine/Plugins/Online/Experimental/OnlineServicesNull/Source/Private/Online/OnlineServicesNull.cpp
+++ b/Engine/Plugins/Online/Experimental/OnlineServicesNull/Source/Private/Online/OnlineServicesNull.cpp
@@ -4,6 +4,7 @@
 #include "Online/OnlineServicesNullTypes.h"
 
 #include "Online/AuthNull.h"
+#include "Online/LobbiesNull.h"
 
 namespace UE::Online {
 
@@ -30,9 +31,44 @@ FOnlineServicesNull::FOnlineServicesNull()
 void FOnlineServicesNull::RegisterComponents()
 {
 	Components.Register<FAuthNull>(*this);
+	Components.Register<FLobbiesNull>(*this);
 	FOnlineServicesCommon::RegisterComponents();
 }
 
+TOnlineResult<FGetResolvedConnectString> FOnlineServicesNull::GetResolvedConnectString(FGetResolvedConnectString::Params&& Params)
+{
+	ILobbiesPtr LobbiesPtr = GetLobbiesInterface();
+	check(LobbiesPtr);
+
+	TOnlineResult<FGetJoinedLobbies> JoinedLobbies = LobbiesPtr->GetJoinedLobbies({ Params.LocalUserId });
+	if (JoinedLobbies.IsOk())
+	{
+		for (TSharedRef<const FLobby>& Lobby : JoinedLobbies.GetOkValue().Lobbies)
+		{
+			if (Lobby->LobbyId == Params.LobbyId)
+			{
+				FName Tag = FName(TEXT("ConnectAddress")); // todo: global tag
+				if(Lobby->Attributes.Contains(Tag))
+				{
+					FGetResolvedConnectString::Result Result;
+					Result.ResolvedConnectString = Lobby->Attributes[Tag].GetString();
+ 					return TOnlineResult<FGetResolvedConnectString>(Result);
+				}
+				else
+				{
+					continue;
+				}
+			}
+		}
+		// No matching lobby
+		return TOnlineResult<FGetResolvedConnectString>(Errors::InvalidParams());
+	}
+	else
+	{
+		return TOnlineResult<FGetResolvedConnectString>(JoinedLobbies.GetErrorValue());
+	}
+}
+
 void FOnlineServicesNull::Initialize()
 {
 	FNullPlatformConfig NullPlatformConfig;
@@ -46,3 +82,4 @@ void FOnlineServicesNull::Initialize()
 
 
 /* UE::Online */ }
+
diff --git a/Engine/Plugins/Online/Experimental/OnlineServicesNull/Source/Private/Online/OnlineServicesNullModule.cpp b/Engine/Plugins/Online/Experimental/OnlineServicesNull/Source/Private/Online/OnlineServicesNullModule.cpp
--- a/Engine/Plugins/Online/Experimental/OnlineServicesNull/Source/Private/Online/OnlineServicesNullModule.cpp
+++ b/Engine/Plugins/Online/Experimental/OnlineServicesNull/Source/Private/Online/OnlineServicesNullModule.cpp
@@ -6,6 +6,7 @@
 
 #include "Online/OnlineServicesRegistry.h"
 #include "Online/OnlineServicesNull.h"
+#include "Online/AuthNull.h"
 
 namespace UE::Online
 {
diff --git a/Engine/Plugins/Online/Experimental/OnlineServicesNull/Source/Public/Online/AuthNull.h b/Engine/Plugins/Online/Experimental/OnlineServicesNull/Source/Public/Online/AuthNull.h
--- a/Engine/Plugins/Online/Experimental/OnlineServicesNull/Source/Public/Online/AuthNull.h
+++ b/Engine/Plugins/Online/Experimental/OnlineServicesNull/Source/Public/Online/AuthNull.h
@@ -29,7 +29,7 @@ public:
 	// Begin IOnlineAccountIdRegistry
 	virtual FString ToLogString(const FOnlineAccountIdHandle& Handle) const override;
 	virtual TArray<uint8> ToReplicationData(const FOnlineAccountIdHandle& Handle) const override;
-	virtual FOnlineAccountIdHandle FromReplicationData(const TArray<uint8> ReplicationString) override;
+	virtual FOnlineAccountIdHandle FromReplicationData(const TArray<uint8>& ReplicationString) override;
 	// End IOnlineAccountIdRegistry
 
 	virtual ~FOnlineAccountIdRegistryNull() = default;
diff --git a/Engine/Plugins/Online/Experimental/OnlineServicesNull/Source/Public/Online/LobbiesNull.h b/Engine/Plugins/Online/Experimental/OnlineServicesNull/Source/Public/Online/LobbiesNull.h
new file mode 100644
--- /dev/null
+++ b/Engine/Plugins/Online/Experimental/OnlineServicesNull/Source/Public/Online/LobbiesNull.h
@@ -0,0 +1,153 @@
+// Copyright Epic Games, Inc. All Rights Reserved.
+#pragma once
+
+#include "CoreMinimal.h"
+#include "Online/LobbiesCommon.h"
+#include "Online/Lobbies.h"
+#include "LANBeacon.h"
+
+namespace UE::Online {
+
+class FLobbyNull
+{
+public: 
+	FLobbyNull();
+
+public:
+	TSharedRef<FLobby> Data;
+	int32 PingInMs;
+	TSharedPtr<class FInternetAddr> HostAddr;
+	uint32 HostAddrIp;
+	int32 HostAddrPort;
+};
+
+class FLobbiesNull : public FLobbiesCommon
+{
+
+public:
+
+	FLobbiesNull(FOnlineServicesNull& InServices)
+		: FLobbiesCommon(InServices)
+		, Services(InServices) {}
+
+	virtual void Initialize() override;
+	virtual void PreShutdown() override;
+	virtual void Tick(float DeltaSeconds) override;
+
+	virtual TOnlineAsyncOpHandle<FCreateLobby> CreateLobby(FCreateLobby::Params&& Params) override;
+	virtual TOnlineAsyncOpHandle<FFindLobbies> FindLobbies(FFindLobbies::Params&& Params) override;
+	virtual TOnlineAsyncOpHandle<FJoinLobby> JoinLobby(FJoinLobby::Params&& Params) override;
+	virtual TOnlineAsyncOpHandle<FLeaveLobby> LeaveLobby(FLeaveLobby::Params&& Params) override;
+	virtual TOnlineResult<FGetJoinedLobbies> GetJoinedLobbies(FGetJoinedLobbies::Params&& Params) override;
+private:
+	// LAN
+	bool NeedsToAdvertise();
+	bool IsSessionJoinable( const TSharedRef<FLobbyNull>& Session) const;
+	uint32 UpdateLANStatus();
+	uint32 JoinLANSession(int32 PlayerNum, TSharedRef<FLobbyNull>& Lobby);
+	uint32 FindLANSession();
+	uint32 FinalizeLANSearch();
+	void AppendLobbyToPacket(class FNboSerializeToBufferNullSvc& Packet, const TSharedRef<FLobbyNull>& Lobby);
+	//void AppendSessionSettingsToPacket(class FNboSerializeToBufferNull& Packet, FOnlineSessionSettings* SessionSettings);
+	void ReadLobbyFromPacket(class FNboSerializeFromBufferNullSvc& Packet, const TSharedRef<FLobbyNull>& Session);
+	//void ReadSettingsFromPacket(class FNboSerializeFromBufferNull& Packet, FOnlineSessionSettings& SessionSettings);
+	void OnValidQueryPacketReceived(uint8* PacketData, int32 PacketLength, uint64 ClientNonce);
+	void OnValidResponsePacketReceived(uint8* PacketData, int32 PacketLength);
+	void OnLANSearchTimeout();
+
+private:
+	FOnlineServicesNull& Services;
+	
+	TMap<FName, TSharedRef<FLobbyNull>> NamedLobbies;
+	TMap<FOnlineLobbyIdHandle, TSharedRef<FLobbyNull>> AllLobbies; // todo: this should be moved to the lobby id registry (replication made this tricky to get implemented asap)
+	TSharedPtr<FFindLobbies::Result> CurrentLobbySearch;
+	TSharedPtr<TOnlineAsyncOp<FFindLobbies>> CurrentLobbySearchHandle;
+	FLANSession LANSessionManager;
+
+
+
+	// Lobby list management
+	TSharedRef<FLobbyNull> CreateNamedLobby(const FCreateLobby::Params& Params);
+
+	TSharedPtr<const FLobbyNull> GetNamedLobby(const FName& LobbyName) const
+	{
+		const TSharedRef<FLobbyNull>* Lobby = NamedLobbies.Find(LobbyName);
+		if (Lobby)
+		{
+			return *Lobby;
+		}
+		else
+		{
+			return nullptr;
+		}
+	}
+
+	TSharedPtr<FLobbyNull> GetNamedLobby(const FName& LobbyName)
+	{
+		TSharedRef<FLobbyNull>* Lobby = NamedLobbies.Find(LobbyName);
+		if(Lobby)
+		{
+			return *Lobby;
+		}
+		else 
+		{
+			return nullptr;
+		}
+	}
+
+	TSharedPtr<FLobbyNull> GetLobby(const FOnlineLobbyIdHandle LobbyId)
+	{
+		TSharedRef<FLobbyNull>* Lobby = AllLobbies.Find(LobbyId);
+		if(Lobby)
+		{
+			return *Lobby;
+		}
+		else
+		{
+			return nullptr;
+		}
+	}
+	
+	void AddNamedLobby(const TSharedRef<FLobbyNull>& Lobby, FName Name)
+	{
+		//checkf(!NamedLobbies.Contains(Name));
+		NamedLobbies.Add(Name, Lobby);
+		AllLobbies.Add(Lobby->Data->LobbyId, Lobby); // Should be redundant, theoretically, but will help clear up any potential edge cases
+		Lobby->Data->LocalName = Name;
+	}
+	
+	// returns true if a named lobby was removed
+	bool RemoveLobbyFromRef(const TSharedRef<FLobbyNull>& Lobby)
+	{
+		return RemoveLobbyFromName(Lobby->Data->LocalName);
+	}
+
+	// returns true if a named lobby was removed
+	bool RemoveLobbyFromId(const FOnlineLobbyIdHandle& LobbyId)
+	{
+		TSharedPtr<FLobbyNull> Lobby = GetLobby(LobbyId);
+		if(Lobby)
+		{
+			int32 NumRemoved = NamedLobbies.Remove(Lobby->Data->LocalName);
+			AllLobbies.Remove(LobbyId);
+			return NumRemoved > 0;
+		}
+		return false;
+	}
+	
+	// returns true if a named lobby was removed
+	bool RemoveLobbyFromName(const FName& Name)
+	{
+		TSharedPtr<FLobbyNull> Lobby = GetNamedLobby(Name);
+		if(Lobby)
+		{
+			NamedLobbies.Remove(Name);
+			AllLobbies.Remove(Lobby->Data->LobbyId);
+			return true;
+		}
+		return false;
+	}
+
+};
+
+} // namespace UE::Online
\ No newline at end of file
diff --git a/Engine/Plugins/Online/Experimental/OnlineServicesNull/Source/Public/Online/LobbyRegistryNull.h b/Engine/Plugins/Online/Experimental/OnlineServicesNull/Source/Public/Online/LobbyRegistryNull.h
new file mode 100644
--- /dev/null
+++ b/Engine/Plugins/Online/Experimental/OnlineServicesNull/Source/Public/Online/LobbyRegistryNull.h
@@ -0,0 +1,34 @@
+// Copyright Epic Games, Inc. All Rights Reserved.
+
+#pragma once
+
+#include "CoreMinimal.h"
+#include "Online/LobbiesCommon.h"
+
+namespace UE::Online {
+using IOnlineLobbyIdRegistry = IOnlineIdRegistry<OnlineIdHandleTags::FLobby>; // todo: remove this when its added to the global scope
+class FOnlineLobbyIdRegistryNull : public IOnlineLobbyIdRegistry
+{
+public:
+	static FOnlineLobbyIdRegistryNull& Get();
+
+	const FOnlineLobbyIdHandle* Find(FString LobbyId);
+	FOnlineLobbyIdHandle FindOrAdd(FString LobbyId);
+	FOnlineLobbyIdHandle GetNext();
+
+	// Begin IOnlineAccountIdRegistry
+	virtual FString ToLogString(const FOnlineLobbyIdHandle& Handle) const override;
+	virtual TArray<uint8> ToReplicationData(const FOnlineLobbyIdHandle& Handle) const override;
+	virtual FOnlineLobbyIdHandle FromReplicationData(const TArray<uint8>& ReplicationString) override;
+	// End IOnlineAccountIdRegistry
+
+	virtual ~FOnlineLobbyIdRegistryNull() = default;
+
+private:
+	const FString* GetInternal(const FOnlineLobbyIdHandle& Handle) const;
+	TArray<FString> Ids;
+	TMap<FString, FOnlineLobbyIdHandle> StringToId; 
+
+};
+
+} // namespace UE::Online
\ No newline at end of file
diff --git a/Engine/Plugins/Online/Experimental/OnlineServicesNull/Source/Public/Online/NboSerializerNullSvc.h b/Engine/Plugins/Online/Experimental/OnlineServicesNull/Source/Public/Online/NboSerializerNullSvc.h
new file mode 100644
--- /dev/null
+++ b/Engine/Plugins/Online/Experimental/OnlineServicesNull/Source/Public/Online/NboSerializerNullSvc.h
@@ -0,0 +1,95 @@
+// Copyright Epic Games, Inc. All Rights Reserved.
+
+#pragma once
+
+#include "CoreMinimal.h"
+#include "NboSerializer.h"
+#include "Online/AuthNull.h"
+#include "Online/CoreOnline.h"
+
+/**
+ * Serializes data in network byte order form into a buffer
+ */
+ namespace UE::Online {
+
+class FNboSerializeToBufferNullSvc : public FNboSerializeToBuffer
+{
+public:
+	/** Default constructor zeros num bytes*/
+	FNboSerializeToBufferNullSvc() :
+		FNboSerializeToBuffer(512)
+	{
+	}
+
+	/** Constructor specifying the size to use */
+	FNboSerializeToBufferNullSvc(uint32 Size) :
+		FNboSerializeToBuffer(Size)
+	{
+	}
+
+	/**
+	 * Adds Null Id to the buffer
+	 */
+	friend inline FNboSerializeToBufferNullSvc& operator<<(FNboSerializeToBufferNullSvc& Ar, const FOnlineAccountIdHandle& UniqueId)
+	{
+		TArray<uint8> Data = FOnlineAccountIdRegistryNull::Get().ToReplicationData(UniqueId);
+		Ar << Data.Num();
+		Ar.WriteBinary(Data.GetData(), Data.Num());
+		return Ar;
+	}
+
+	friend inline FNboSerializeToBufferNullSvc& operator<<(FNboSerializeToBufferNullSvc& Ar, const TMap<FLobbyAttributeId, FLobbyVariant>& Map)
+	{
+		Ar << Map.Num();
+		for (const TPair<FLobbyAttributeId, FLobbyVariant>& Pair : Map)
+		{
+			Ar << Pair.Key;
+			Ar << Pair.Value.GetString();
+		}
+		return Ar;
+	}
+};
+
+/**
+ * Class used to write data into packets for sending via system link
+ */
+class FNboSerializeFromBufferNullSvc : public FNboSerializeFromBuffer
+{
+public:
+	/**
+	 * Initializes the buffer, size, and zeros the read offset
+	 */
+	FNboSerializeFromBufferNullSvc(uint8* Packet,int32 Length) :
+		FNboSerializeFromBuffer(Packet,Length)
+	{
+	}
+	/**
+	 * Reads Null Id from the buffer
+	 */
+	friend inline FNboSerializeFromBufferNullSvc& operator>>(FNboSerializeFromBufferNullSvc& Ar, FOnlineAccountIdHandle& UniqueId)
+	{
+		TArray<uint8> Data;
+		int32 Size;
+		Ar >> Size;
+		Ar.ReadBinaryArray(Data, Size);
+		UniqueId = FOnlineAccountIdRegistryNull::Get().FromReplicationData(Data);
+		return Ar;
+	}
+	
+	friend inline FNboSerializeFromBufferNullSvc& operator>>(FNboSerializeFromBufferNullSvc& Ar, TMap<FLobbyAttributeId, FLobbyVariant>& Map)
+	{
+		int32 Size;
+		Ar >> Size;
+		for(int i = 0; i < Size; i++)
+		{
+			FLobbyAttributeId LobbyAttributeId;
+			FString LobbyData;
+			Ar >> LobbyAttributeId;
+			Ar >> LobbyData;
+			Map.Emplace(MoveTemp(LobbyAttributeId), MoveTemp(LobbyData));
+		}
+		return Ar;
+	}
+};
+
+} //namespace UE::Online
\ No newline at end of file
diff --git a/Engine/Plugins/Online/Experimental/OnlineServicesNull/Source/Public/Online/OnlineServicesNull.h b/Engine/Plugins/Online/Experimental/OnlineServicesNull/Source/Public/Online/OnlineServicesNull.h
--- a/Engine/Plugins/Online/Experimental/OnlineServicesNull/Source/Public/Online/OnlineServicesNull.h
+++ b/Engine/Plugins/Online/Experimental/OnlineServicesNull/Source/Public/Online/OnlineServicesNull.h
@@ -14,6 +14,7 @@ public:
 	FOnlineServicesNull();
 	virtual void RegisterComponents() override;
 	virtual void Initialize() override;
+	virtual TOnlineResult<FGetResolvedConnectString> GetResolvedConnectString(FGetResolvedConnectString::Params&& Params) override;
 
 	static EOnlineServices GetServicesProvider() { return EOnlineServices::Null; }
 };
