Fix build

diff --git a/Engine/Plugins/Online/Experimental/OnlineServicesNull/Source/Private/Online/LobbiesNull.cpp b/Engine/Plugins/Online/Experimental/OnlineServicesNull/Source/Private/Online/LobbiesNull.cpp
--- a/Engine/Plugins/Online/Experimental/OnlineServicesNull/Source/Private/Online/LobbiesNull.cpp
+++ b/Engine/Plugins/Online/Experimental/OnlineServicesNull/Source/Private/Online/LobbiesNull.cpp
@@ -1,6 +1,7 @@
 // Copyright Epic Games, Inc. All Rights Reserved.
 
 #include "Online/LobbiesNull.h"
+#include "Online/OnlineServicesNull.h"
 #include "Misc/Guid.h"
 #include "SocketSubsystem.h"
 #include "Online/NboSerializerNullSvc.h"
@@ -42,6 +43,12 @@ FLobbyNull::FLobbyNull()
 
 	HostAddr->GetIp(HostAddrIp);
 	HostAddr->GetPort(HostAddrPort);
+}
+
+FLobbiesNull::FLobbiesNull(FOnlineServicesNull& InServices)
+	: FLobbiesCommon(InServices)
+	, Services(InServices)
+{
 
 }
 
diff --git a/Engine/Plugins/Online/Experimental/OnlineServicesNull/Source/Public/Online/LobbiesNull.h b/Engine/Plugins/Online/Experimental/OnlineServicesNull/Source/Public/Online/LobbiesNull.h
--- a/Engine/Plugins/Online/Experimental/OnlineServicesNull/Source/Public/Online/LobbiesNull.h
+++ b/Engine/Plugins/Online/Experimental/OnlineServicesNull/Source/Public/Online/LobbiesNull.h
@@ -8,6 +8,8 @@
 
 namespace UE::Online {
 
+class FOnlineServicesNull;
+
 class FLobbyNull
 {
 public: 
@@ -26,9 +28,7 @@ class FLobbiesNull : public FLobbiesCommon
 
 public:
 
-	FLobbiesNull(FOnlineServicesNull& InServices)
-		: FLobbiesCommon(InServices)
-		, Services(InServices) {}
+	FLobbiesNull(FOnlineServicesNull& InServices);
 
 	virtual void Initialize() override;
 	virtual void PreShutdown() override;
