Allow OnlineServicesNull to compile without engine

diff --git a/Engine/Plugins/Online/Experimental/OnlineServicesNull/Source/OnlineServicesNull.Build.cs b/Engine/Plugins/Online/Experimental/OnlineServicesNull/Source/OnlineServicesNull.Build.cs
--- a/Engine/Plugins/Online/Experimental/OnlineServicesNull/Source/OnlineServicesNull.Build.cs
+++ b/Engine/Plugins/Online/Experimental/OnlineServicesNull/Source/OnlineServicesNull.Build.cs
@@ -13,10 +13,19 @@ public class OnlineServicesNull : ModuleRules
 				"CoreOnline",
 				"Sockets",
 				"OnlineServicesInterface",
-				"OnlineServicesCommon",
-				"OnlineSubsystem",
-				"OnlineSubsystemUtils"
+				"OnlineServicesCommon"
 			}
 		);
+
+		if (Target.bCompileAgainstEngine)
+		{
+			PublicDependencyModuleNames.AddRange(
+				new string[]
+				{
+					"OnlineSubsystem",
+					"OnlineSubsystemUtils"
+				}
+			);
+		}
 	}
 }
diff --git a/Engine/Plugins/Online/Experimental/OnlineServicesNull/Source/Private/Online/LobbiesNull.cpp b/Engine/Plugins/Online/Experimental/OnlineServicesNull/Source/Private/Online/LobbiesNull.cpp
--- a/Engine/Plugins/Online/Experimental/OnlineServicesNull/Source/Private/Online/LobbiesNull.cpp
+++ b/Engine/Plugins/Online/Experimental/OnlineServicesNull/Source/Private/Online/LobbiesNull.cpp
@@ -1,5 +1,6 @@
 // Copyright Epic Games, Inc. All Rights Reserved.
 
+#if WITH_ENGINE
 #include "Online/LobbiesNull.h"
 #include "Online/OnlineServicesNull.h"
 #include "Misc/Guid.h"
@@ -410,3 +411,5 @@ TOnlineResult<FGetJoinedLobbies> FLobbiesNull::GetJoinedLobbies(FGetJoinedLobbie
 
 
 #undef NOT_IMPLEMENTED
+
+#endif // WITH_ENGINE
diff --git a/Engine/Plugins/Online/Experimental/OnlineServicesNull/Source/Private/Online/OnlineServicesNull.cpp b/Engine/Plugins/Online/Experimental/OnlineServicesNull/Source/Private/Online/OnlineServicesNull.cpp
--- a/Engine/Plugins/Online/Experimental/OnlineServicesNull/Source/Private/Online/OnlineServicesNull.cpp
+++ b/Engine/Plugins/Online/Experimental/OnlineServicesNull/Source/Private/Online/OnlineServicesNull.cpp
@@ -31,12 +31,15 @@ FOnlineServicesNull::FOnlineServicesNull()
 void FOnlineServicesNull::RegisterComponents()
 {
 	Components.Register<FAuthNull>(*this);
+#if WITH_ENGINE
 	Components.Register<FLobbiesNull>(*this);
+#endif
 	FOnlineServicesCommon::RegisterComponents();
 }
 
 TOnlineResult<FGetResolvedConnectString> FOnlineServicesNull::GetResolvedConnectString(FGetResolvedConnectString::Params&& Params)
 {
+#if WITH_ENGINE
 	ILobbiesPtr LobbiesPtr = GetLobbiesInterface();
 	check(LobbiesPtr);
 
@@ -67,6 +70,9 @@ TOnlineResult<FGetResolvedConnectString> FOnlineServicesNull::GetResolvedConnect
 	{
 		return TOnlineResult<FGetResolvedConnectString>(JoinedLobbies.GetErrorValue());
 	}
+#else
+	return Super::GetResolvedConnectString(MoveTemp(Params));
+#endif
 }
 
 void FOnlineServicesNull::Initialize()
diff --git a/Engine/Plugins/Online/Experimental/OnlineServicesNull/Source/Public/Online/LobbiesNull.h b/Engine/Plugins/Online/Experimental/OnlineServicesNull/Source/Public/Online/LobbiesNull.h
--- a/Engine/Plugins/Online/Experimental/OnlineServicesNull/Source/Public/Online/LobbiesNull.h
+++ b/Engine/Plugins/Online/Experimental/OnlineServicesNull/Source/Public/Online/LobbiesNull.h
@@ -1,6 +1,8 @@
 // Copyright Epic Games, Inc. All Rights Reserved.
 #pragma once
 
+#if WITH_ENGINE
+
 #include "CoreMinimal.h"
 #include "Online/LobbiesCommon.h"
 #include "Online/Lobbies.h"
@@ -150,4 +152,6 @@ private:
 
 };
 
-} // namespace UE::Online
\ No newline at end of file
+} // namespace UE::Online
+
+#endif // WITH_ENGINE
