Ensure ID handle is valid in GetInternal for the auth null lookup

diff --git a/Engine/Plugins/Online/Experimental/OnlineServicesNull/Source/Private/Online/AuthNull.cpp b/Engine/Plugins/Online/Experimental/OnlineServicesNull/Source/Private/Online/AuthNull.cpp
--- a/Engine/Plugins/Online/Experimental/OnlineServicesNull/Source/Private/Online/AuthNull.cpp
+++ b/Engine/Plugins/Online/Experimental/OnlineServicesNull/Source/Private/Online/AuthNull.cpp
@@ -206,7 +206,7 @@ FOnlineAccountIdHandle FOnlineAccountIdRegistryNull::Find(int32 UserId) const
 
 const FOnlineAccountIdString* FOnlineAccountIdRegistryNull::GetInternal(const FOnlineAccountIdHandle& Handle) const
 {
-	if(Handle.GetOnlineServicesType() == EOnlineServices::Null && Handle.GetHandle() <= (uint32)Ids.Num())
+	if(Handle.IsValid() && Handle.GetOnlineServicesType() == EOnlineServices::Null && Handle.GetHandle() <= (uint32)Ids.Num())
 	{
 		return &Ids[Handle.GetHandle()-1];
 	}
