- Achievements Null implementation

diff --git a/Engine/Plugins/Online/Experimental/OnlineServicesNull/Source/Private/Online/AchievementsNull.cpp b/Engine/Plugins/Online/Experimental/OnlineServicesNull/Source/Private/Online/AchievementsNull.cpp
new file mode 100644
--- /dev/null
+++ b/Engine/Plugins/Online/Experimental/OnlineServicesNull/Source/Private/Online/AchievementsNull.cpp
@@ -0,0 +1,275 @@
+// Copyright Epic Games, Inc. All Rights Reserved.
+
+#include "Online/AchievementsNull.h"
+
+#include "Math/UnrealMathUtility.h"
+#include "Online/AchievementsErrors.h"
+#include "Online/AuthNull.h"
+#include "Online/OnlineErrorDefinitions.h"
+#include "Online/OnlineServicesNull.h"
+#include "Online/OnlineServicesNullTypes.h"
+#include "CoreGlobals.h"
+
+namespace UE::Online {
+
+namespace {
+
+const TCHAR* ConfigSection = TEXT("OnlineServices.Null.Achievements");
+
+TMap<FString, FAchievementDefinition> GetAchievementDefinitionsFromConfig()
+{
+	TMap<FString, FAchievementDefinition> Result;
+
+	for (int AchievIdx = 0;; AchievIdx++)
+	{
+		FString AchievementId;
+		GConfig->GetString(ConfigSection, *FString::Printf(TEXT("AchievementDef_%d_Id"), AchievIdx), AchievementId, GEngineIni);
+		if (AchievementId.IsEmpty())
+		{
+			break;
+		}
+
+		FAchievementDefinition& AchievementDefinition = Result.Emplace(AchievementId);
+		AchievementDefinition.AchievementId = MoveTemp(AchievementId);
+
+		GConfig->GetText(ConfigSection, *FString::Printf(TEXT("AchievementDef_%d_UnlockedDisplayName"), AchievIdx), AchievementDefinition.UnlockedDisplayName, GEngineIni);
+		GConfig->GetText(ConfigSection, *FString::Printf(TEXT("AchievementDef_%d_UnlockedDescription"), AchievIdx), AchievementDefinition.UnlockedDescription, GEngineIni);
+		GConfig->GetText(ConfigSection, *FString::Printf(TEXT("AchievementDef_%d_LockedDisplayName"), AchievIdx), AchievementDefinition.LockedDisplayName, GEngineIni);
+		GConfig->GetText(ConfigSection, *FString::Printf(TEXT("AchievementDef_%d_LockedDescription"), AchievIdx), AchievementDefinition.LockedDescription, GEngineIni);
+		GConfig->GetText(ConfigSection, *FString::Printf(TEXT("AchievementDef_%d_FlavorText"), AchievIdx), AchievementDefinition.FlavorText, GEngineIni);
+		GConfig->GetString(ConfigSection, *FString::Printf(TEXT("AchievementDef_%d_UnlockedIconUrl"), AchievIdx), AchievementDefinition.UnlockedIconUrl, GEngineIni);
+		GConfig->GetString(ConfigSection, *FString::Printf(TEXT("AchievementDef_%d_LockedIconUrl"), AchievIdx), AchievementDefinition.LockedIconUrl, GEngineIni);
+		GConfig->GetBool(ConfigSection, *FString::Printf(TEXT("AchievementDef_%d_bIsHidden"), AchievIdx), AchievementDefinition.bIsHidden, GEngineIni);
+
+		for (int StatDefIdx = 0;; StatDefIdx++)
+		{
+			FString StatId;
+			GConfig->GetString(ConfigSection, *FString::Printf(TEXT("AchievementDef_%d_StatDef_%d_Id"), AchievIdx, StatDefIdx), StatId, GEngineIni);
+			if (StatId.IsEmpty())
+			{
+				break;
+			}
+
+			FAchievementStatDefinition StatDefinition = AchievementDefinition.StatDefinitions.Emplace_GetRef();
+			StatDefinition.StatId = MoveTemp(StatId);
+
+			int32 UnlockThreshold = 0;
+			GConfig->GetInt(ConfigSection, *FString::Printf(TEXT("AchievementDef_%d_StatDef_%d_UnlockThreshold"), AchievIdx, StatDefIdx), UnlockThreshold, GEngineIni);
+			StatDefinition.UnlockThreshold = UnlockThreshold;
+		}
+	}
+
+	return Result;
+}
+
+}
+
+FAchievementsNull::FAchievementsNull(FOnlineServicesNull& InOwningSubsystem)
+	: Super(InOwningSubsystem)
+{
+}
+
+TOnlineAsyncOpHandle<FQueryAchievementDefinitions> FAchievementsNull::QueryAchievementDefinitions(FQueryAchievementDefinitions::Params&& Params)
+{
+	TOnlineAsyncOpRef<FQueryAchievementDefinitions> Op = GetOp<FQueryAchievementDefinitions>(MoveTemp(Params));
+
+	if (!Services.Get<FAuthNull>()->IsLoggedIn(Op->GetParams().LocalUserId))
+	{
+		Op->SetError(Errors::InvalidUser());
+		return Op->GetHandle();
+	}
+
+	if(!AchievementDefinitions.IsSet())
+	{
+		AchievementDefinitions.Emplace(GetAchievementDefinitionsFromConfig());
+	}
+
+	Op->SetResult({});
+	return Op->GetHandle();
+}
+
+TOnlineResult<FGetAchievementIds> FAchievementsNull::GetAchievementIds(FGetAchievementIds::Params&& Params)
+{
+	if (!Services.Get<FAuthNull>()->IsLoggedIn(Params.LocalUserId))
+	{
+		return TOnlineResult<FGetAchievementIds>(Errors::InvalidUser());
+	}
+
+	if (!AchievementDefinitions.IsSet())
+	{
+		// Call QueryAchievementDefinitions first
+		return TOnlineResult<FGetAchievementIds>(Errors::InvalidState());
+	}
+
+	FGetAchievementIds::Result Result;
+	AchievementDefinitions->GenerateKeyArray(Result.AchievementIds);
+	return TOnlineResult<FGetAchievementIds>(MoveTemp(Result));
+}
+
+TOnlineResult<FGetAchievementDefinition> FAchievementsNull::GetAchievementDefinition(FGetAchievementDefinition::Params&& Params)
+{
+	if (!Services.Get<FAuthNull>()->IsLoggedIn(Params.LocalUserId))
+	{
+		return TOnlineResult<FGetAchievementDefinition>(Errors::InvalidUser());
+	}
+
+	if (!AchievementDefinitions.IsSet())
+	{
+		// Should call QueryAchievementDefinitions first
+		return TOnlineResult<FGetAchievementDefinition>(Errors::InvalidState());
+	}
+
+	const FAchievementDefinition* AchievementDefinition = AchievementDefinitions->Find(Params.AchievementId);
+	if (!AchievementDefinition)
+	{
+		return TOnlineResult<FGetAchievementDefinition>(Errors::NotFound());
+	}
+
+	return TOnlineResult<FGetAchievementDefinition>({ *AchievementDefinition });
+}
+
+TOnlineAsyncOpHandle<FQueryAchievementStates> FAchievementsNull::QueryAchievementStates(FQueryAchievementStates::Params&& Params)
+{
+	TOnlineAsyncOpRef<FQueryAchievementStates> Op = GetOp<FQueryAchievementStates>(MoveTemp(Params));
+
+	if (!Services.Get<FAuthNull>()->IsLoggedIn(Op->GetParams().LocalUserId))
+	{
+		Op->SetError(Errors::InvalidUser());
+		return Op->GetHandle();
+	}
+
+	if (!AchievementDefinitions.IsSet())
+	{
+		// Call QueryAchievementDefinitions first
+		Op->SetError(Errors::InvalidState());
+		return Op->GetHandle();
+	}
+
+	if (!AchievementStates.Contains(Op->GetParams().LocalUserId))
+	{
+		FAchievementStateMap& LocalUserAchievementStates = AchievementStates.Emplace(Op->GetParams().LocalUserId);
+		for (const TPair<FString, FAchievementDefinition>& Pair : *AchievementDefinitions)
+		{
+			const FString& AchievementId = Pair.Key;
+			FAchievementState& AchievementState = LocalUserAchievementStates.Emplace(AchievementId);
+			AchievementState.AchievementId = AchievementId;
+		}
+	}
+
+	Op->SetResult({});
+	return Op->GetHandle();
+}
+
+TOnlineResult<FGetAchievementState> FAchievementsNull::GetAchievementState(FGetAchievementState::Params&& Params)
+{
+	if (!Services.Get<FAuthNull>()->IsLoggedIn(Params.LocalUserId))
+	{
+		return TOnlineResult<FGetAchievementState>(Errors::InvalidUser());
+	}
+
+	const FAchievementStateMap* LocalUserAchievementStates = AchievementStates.Find(Params.LocalUserId);
+	if (!LocalUserAchievementStates)
+	{
+		// Call QueryAchievementStates first
+		return TOnlineResult<FGetAchievementState>(Errors::InvalidState());
+	}
+
+	const FAchievementState* AchievementState = LocalUserAchievementStates->Find(Params.AchievementId);
+	if (!AchievementState)
+	{
+		return TOnlineResult<FGetAchievementState>(Errors::NotFound());
+	}
+
+	return TOnlineResult<FGetAchievementState>({*AchievementState});
+}
+
+TOnlineAsyncOpHandle<FUnlockAchievements> FAchievementsNull::UnlockAchievements(FUnlockAchievements::Params&& Params)
+{
+	TOnlineAsyncOpRef<FUnlockAchievements> Op = GetOp<FUnlockAchievements>(MoveTemp(Params));
+
+	if (!Services.Get<FAuthNull>()->IsLoggedIn(Op->GetParams().LocalUserId))
+	{
+		Op->SetError(Errors::InvalidUser());
+		return Op->GetHandle();
+	}
+
+	if (Op->GetParams().AchievementIds.IsEmpty())
+	{
+		Op->SetError(Errors::InvalidParams());
+		return Op->GetHandle();
+	}
+
+	FAchievementStateMap* LocalUserAchievementStates = AchievementStates.Find(Op->GetParams().LocalUserId);
+	if (!LocalUserAchievementStates)
+	{
+		// Call QueryAchievementStates first
+		Op->SetError(Errors::InvalidState());
+		return Op->GetHandle();
+	}
+
+	for (const FString& AchievementId : Op->GetParams().AchievementIds)
+	{
+		const FAchievementState* AchievementState = LocalUserAchievementStates->Find(AchievementId);
+		if (!AchievementState)
+		{
+			Op->SetError(Errors::NotFound());
+			return Op->GetHandle();
+		}
+		if (FMath::IsNearlyEqual(AchievementState->Progress, 1.0f))
+		{
+			Op->SetError(Errors::Achievements::AlreadyUnlocked());
+			return Op->GetHandle();
+		}
+	}
+
+	FDateTime UtcNow = FDateTime::UtcNow();
+
+	for (const FString& AchievementId : Op->GetParams().AchievementIds)
+	{
+		FAchievementState* AchievementState = LocalUserAchievementStates->Find(AchievementId);
+		if(ensure(AchievementState))
+		{
+			AchievementState->Progress = 1.0f;
+			AchievementState->UnlockTime = UtcNow;
+		}
+	}
+
+	FAchievementStateUpdated AchievementStateUpdated;
+	AchievementStateUpdated.LocalUserId = Op->GetParams().LocalUserId;
+	AchievementStateUpdated.AchievementIds = Op->GetParams().AchievementIds;
+	OnAchievementStateUpdatedEvent.Broadcast(AchievementStateUpdated);
+
+	Op->SetResult({});
+	return Op->GetHandle();
+}
+
+TOnlineResult<FDisplayAchievementUI> FAchievementsNull::DisplayAchievementUI(FDisplayAchievementUI::Params&& Params)
+{
+	if (!Services.Get<FAuthNull>()->IsLoggedIn(Params.LocalUserId))
+	{
+		return TOnlineResult<FDisplayAchievementUI>(Errors::InvalidUser());
+	}
+
+	const FAchievementStateMap* LocalUserAchievementStates = AchievementStates.Find(Params.LocalUserId);
+	if (!LocalUserAchievementStates)
+	{
+		// Call QueryAchievementStates first
+		return TOnlineResult<FDisplayAchievementUI>(Errors::InvalidState());
+	}
+
+	// Safe to assume they called QueryAchievementDefinitions from this point, as that is a prereq for QueryAchievementStates.
+
+	const FAchievementDefinition* AchievementDefinition = AchievementDefinitions->Find(Params.AchievementId);
+	if (!AchievementDefinition)
+	{
+		return TOnlineResult<FDisplayAchievementUI>(Errors::NotFound());
+	}
+
+	const FAchievementState& AchievementState = LocalUserAchievementStates->FindChecked(Params.AchievementId);
+	UE_LOG(LogTemp, Display, TEXT("AchievementsNull: DisplayAchievementUI LocalUserId=[%s]"), *ToLogString(Params.LocalUserId));
+	UE_LOG(LogTemp, Display, TEXT("AchievementsNull: DisplayAchievementUI AchievementDefinition=[%s]"), *ToLogString(*AchievementDefinition));
+	UE_LOG(LogTemp, Display, TEXT("AchievementsNull: DisplayAchievementUI AchievementState=[%s]"), *ToLogString(AchievementState));
+	return TOnlineResult<FDisplayAchievementUI>(Errors::Success());
+}
+
+/* UE::Online */ }
diff --git a/Engine/Plugins/Online/Experimental/OnlineServicesNull/Source/Private/Online/OnlineServicesNull.cpp b/Engine/Plugins/Online/Experimental/OnlineServicesNull/Source/Private/Online/OnlineServicesNull.cpp
--- a/Engine/Plugins/Online/Experimental/OnlineServicesNull/Source/Private/Online/OnlineServicesNull.cpp
+++ b/Engine/Plugins/Online/Experimental/OnlineServicesNull/Source/Private/Online/OnlineServicesNull.cpp
@@ -3,6 +3,7 @@
 #include "Online/OnlineServicesNull.h"
 #include "Online/OnlineServicesNullTypes.h"
 
+#include "Online/AchievementsNull.h"
 #include "Online/AuthNull.h"
 #include "Online/LobbiesNull.h"
 
@@ -30,6 +31,7 @@ FOnlineServicesNull::FOnlineServicesNull(FName InInstanceName)
 
 void FOnlineServicesNull::RegisterComponents()
 {
+	Components.Register<FAchievementsNull>(*this);
 	Components.Register<FAuthNull>(*this);
 #if WITH_ENGINE
 	Components.Register<FLobbiesNull>(*this);
diff --git a/Engine/Plugins/Online/Experimental/OnlineServicesNull/Source/Public/Online/AchievementsNull.h b/Engine/Plugins/Online/Experimental/OnlineServicesNull/Source/Public/Online/AchievementsNull.h
new file mode 100644
--- /dev/null
+++ b/Engine/Plugins/Online/Experimental/OnlineServicesNull/Source/Public/Online/AchievementsNull.h
@@ -0,0 +1,36 @@
+// Copyright Epic Games, Inc. All Rights Reserved.
+
+#pragma once
+
+#include "Misc/Optional.h"
+#include "Online/AchievementsCommon.h"
+
+namespace UE::Online {
+
+class FOnlineServicesNull;
+
+class ONLINESERVICESNULL_API FAchievementsNull : public FAchievementsCommon
+{
+public:
+	using Super = FAchievementsCommon;
+
+	FAchievementsNull(FOnlineServicesNull& InOwningSubsystem);
+
+	//IAchievements
+	virtual TOnlineAsyncOpHandle<FQueryAchievementDefinitions> QueryAchievementDefinitions(FQueryAchievementDefinitions::Params&& Params) override;
+	virtual TOnlineResult<FGetAchievementIds> GetAchievementIds(FGetAchievementIds::Params&& Params) override;
+	virtual TOnlineResult<FGetAchievementDefinition> GetAchievementDefinition(FGetAchievementDefinition::Params&& Params) override;
+	virtual TOnlineAsyncOpHandle<FQueryAchievementStates> QueryAchievementStates(FQueryAchievementStates::Params&& Params) override;
+	virtual TOnlineResult<FGetAchievementState> GetAchievementState(FGetAchievementState::Params&& Params) override;
+	virtual TOnlineAsyncOpHandle<FUnlockAchievements> UnlockAchievements(FUnlockAchievements::Params&& Params) override;
+	virtual TOnlineResult<FDisplayAchievementUI> DisplayAchievementUI(FDisplayAchievementUI::Params&& Params) override;
+
+protected:
+	using FAchievementDefinitionMap = TMap<FString, FAchievementDefinition>;
+	using FAchievementStateMap = TMap<FString, FAchievementState>;
+
+	TOptional<FAchievementDefinitionMap> AchievementDefinitions;
+	TMap<FOnlineAccountIdHandle, FAchievementStateMap> AchievementStates;
+};
+
+/* UE::Online */ }
