- TitleFile Null implementation

diff --git a/Engine/Plugins/Online/Experimental/OnlineServicesNull/Source/Private/Online/OnlineServicesNull.cpp b/Engine/Plugins/Online/Experimental/OnlineServicesNull/Source/Private/Online/OnlineServicesNull.cpp
--- a/Engine/Plugins/Online/Experimental/OnlineServicesNull/Source/Private/Online/OnlineServicesNull.cpp
+++ b/Engine/Plugins/Online/Experimental/OnlineServicesNull/Source/Private/Online/OnlineServicesNull.cpp
@@ -6,6 +6,7 @@
 #include "Online/AchievementsNull.h"
 #include "Online/AuthNull.h"
 #include "Online/LobbiesNull.h"
+#include "Online/TitleFileNull.h"
 
 namespace UE::Online {
 
@@ -36,6 +37,7 @@ void FOnlineServicesNull::RegisterComponents()
 #if WITH_ENGINE
 	Components.Register<FLobbiesNull>(*this);
 #endif
+	Components.Register<FTitleFileNull>(*this);
 	FOnlineServicesCommon::RegisterComponents();
 }
 
diff --git a/Engine/Plugins/Online/Experimental/OnlineServicesNull/Source/Private/Online/TitleFileNull.cpp b/Engine/Plugins/Online/Experimental/OnlineServicesNull/Source/Private/Online/TitleFileNull.cpp
new file mode 100644
--- /dev/null
+++ b/Engine/Plugins/Online/Experimental/OnlineServicesNull/Source/Private/Online/TitleFileNull.cpp
@@ -0,0 +1,121 @@
+// Copyright Epic Games, Inc. All Rights Reserved.
+
+#include "Online/TitleFileNull.h"
+
+#include "Containers/StringConv.h"
+#include "Online/AuthNull.h"
+#include "Online/OnlineErrorDefinitions.h"
+#include "Online/OnlineServicesNull.h"
+#include "Online/OnlineServicesNullTypes.h"
+
+namespace {
+
+const TCHAR* ConfigSection = TEXT("OnlineServices.Null.TitleFile");
+
+using FTitleFileContents = UE::Online::FTitleFileContents;
+using FTitleFileContentsRef = UE::Online::FTitleFileContentsRef;
+
+TMap<FString, FTitleFileContentsRef> GetTitleFilesFromConfig()
+{
+	TMap<FString, FTitleFileContentsRef> Result;
+
+	for (int FileIdx = 0;; FileIdx++)
+	{
+		FString Filename;
+		GConfig->GetString(ConfigSection, *FString::Printf(TEXT("File_%d_Name"), FileIdx), Filename, GEngineIni);
+		if (Filename.IsEmpty())
+		{
+			break;
+		}
+
+		FString FileContentsStr;
+		GConfig->GetString(ConfigSection, *FString::Printf(TEXT("File_%d_Contents"), FileIdx), FileContentsStr, GEngineIni);
+
+		if (!FileContentsStr.IsEmpty())
+		{
+			const int ArrayLen = FileContentsStr.GetCharArray().Num() * sizeof(TCHAR);
+			FTitleFileContents FileContents((uint8*)FileContentsStr.GetCharArray().GetData(), ArrayLen);
+
+			Result.Emplace(MoveTemp(Filename), MakeShared<const FTitleFileContents>(MoveTemp(FileContents)));
+		}
+	}
+
+	return Result;
+}
+
+}
+
+namespace UE::Online {
+
+FTitleFileNull::FTitleFileNull(FOnlineServicesNull& InOwningSubsystem)
+	: Super(InOwningSubsystem)
+{
+}
+
+TOnlineAsyncOpHandle<FTitleFileEnumerateFiles> FTitleFileNull::EnumerateFiles(FTitleFileEnumerateFiles::Params&& Params)
+{
+	TOnlineAsyncOpRef<FTitleFileEnumerateFiles> Op = GetOp<FTitleFileEnumerateFiles>(MoveTemp(Params));
+
+	if (!Services.Get<FAuthNull>()->IsLoggedIn(Op->GetParams().LocalUserId))
+	{
+		Op->SetError(Errors::InvalidUser());
+		return Op->GetHandle();
+	}
+
+	if(!TitleFiles.IsSet())
+	{
+		TitleFiles.Emplace(GetTitleFilesFromConfig());
+	}
+
+	Op->SetResult({});
+	return Op->GetHandle();
+}
+
+TOnlineResult<FTitleFileGetEnumeratedFiles> FTitleFileNull::GetEnumeratedFiles(FTitleFileGetEnumeratedFiles::Params&& Params)
+{
+	if (!Services.Get<FAuthNull>()->IsLoggedIn(Params.LocalUserId))
+	{
+		return TOnlineResult<FTitleFileGetEnumeratedFiles>(Errors::InvalidUser());
+	}
+
+	if (!TitleFiles.IsSet())
+	{
+		// Need to call EnumerateFiles first.
+		return TOnlineResult<FTitleFileGetEnumeratedFiles>(Errors::InvalidState());
+	}
+
+	FTitleFileGetEnumeratedFiles::Result Result;
+	TitleFiles->GenerateKeyArray(Result.Filenames);
+	return TOnlineResult<FTitleFileGetEnumeratedFiles>(MoveTemp(Result));
+	
+}
+
+TOnlineAsyncOpHandle<FTitleFileReadFile> FTitleFileNull::ReadFile(FTitleFileReadFile::Params&& Params)
+{
+	TOnlineAsyncOpRef<FTitleFileReadFile> Op = GetOp<FTitleFileReadFile>(MoveTemp(Params));
+
+	if (!Services.Get<FAuthNull>()->IsLoggedIn(Op->GetParams().LocalUserId))
+	{
+		Op->SetError(Errors::InvalidUser());
+		return Op->GetHandle();
+	}
+
+	if (!TitleFiles.IsSet())
+	{
+		// Need to call EnumerateFiles first.
+		Op->SetError(Errors::InvalidState());
+		return Op->GetHandle();
+	}
+
+	const FTitleFileContentsRef* Found = TitleFiles->Find(Op->GetParams().Filename);
+	if (!Found)
+	{
+		Op->SetError(Errors::NotFound());
+		return Op->GetHandle();
+	}
+
+	Op->SetResult({*Found});
+	return Op->GetHandle();
+}
+
+/* UE::Online */ }
diff --git a/Engine/Plugins/Online/Experimental/OnlineServicesNull/Source/Public/Online/TitleFileNull.h b/Engine/Plugins/Online/Experimental/OnlineServicesNull/Source/Public/Online/TitleFileNull.h
new file mode 100644
--- /dev/null
+++ b/Engine/Plugins/Online/Experimental/OnlineServicesNull/Source/Public/Online/TitleFileNull.h
@@ -0,0 +1,28 @@
+// Copyright Epic Games, Inc. All Rights Reserved.
+
+#pragma once
+
+#include "Misc/Optional.h"
+#include "Online/TitleFileCommon.h"
+
+namespace UE::Online {
+
+class FOnlineServicesNull;
+
+class ONLINESERVICESNULL_API FTitleFileNull : public FTitleFileCommon
+{
+public:
+	using Super = FTitleFileCommon;
+
+	FTitleFileNull(FOnlineServicesNull& InOwningSubsystem);
+
+	// ITitleFile
+	virtual TOnlineAsyncOpHandle<FTitleFileEnumerateFiles> EnumerateFiles(FTitleFileEnumerateFiles::Params&& Params) override;
+	virtual TOnlineResult<FTitleFileGetEnumeratedFiles> GetEnumeratedFiles(FTitleFileGetEnumeratedFiles::Params&& Params) override;
+	virtual TOnlineAsyncOpHandle<FTitleFileReadFile> ReadFile(FTitleFileReadFile::Params&& Params) override;
+
+protected:
+	TOptional<TMap<FString, FTitleFileContentsRef>> TitleFiles;
+};
+
+/* UE::Online */ }
