- Fix DisplayAchievementUI returning a success error, rather than a result struct.

diff --git a/Engine/Plugins/Online/Experimental/OnlineServicesNull/Source/Private/Online/AchievementsNull.cpp b/Engine/Plugins/Online/Experimental/OnlineServicesNull/Source/Private/Online/AchievementsNull.cpp
--- a/Engine/Plugins/Online/Experimental/OnlineServicesNull/Source/Private/Online/AchievementsNull.cpp
+++ b/Engine/Plugins/Online/Experimental/OnlineServicesNull/Source/Private/Online/AchievementsNull.cpp
@@ -269,7 +269,8 @@ TOnlineResult<FDisplayAchievementUI> FAchievementsNull::DisplayAchievementUI(FDi
 	UE_LOG(LogTemp, Display, TEXT("AchievementsNull: DisplayAchievementUI LocalUserId=[%s]"), *ToLogString(Params.LocalUserId));
 	UE_LOG(LogTemp, Display, TEXT("AchievementsNull: DisplayAchievementUI AchievementDefinition=[%s]"), *ToLogString(*AchievementDefinition));
 	UE_LOG(LogTemp, Display, TEXT("AchievementsNull: DisplayAchievementUI AchievementState=[%s]"), *ToLogString(AchievementState));
-	return TOnlineResult<FDisplayAchievementUI>(Errors::Success());
+
+	return TOnlineResult<FDisplayAchievementUI>(FDisplayAchievementUI::Result());
 }
 
 /* UE::Online */ }
