Removal of WITH_ENGINE checks from Lobbies after refactor

diff --git a/Engine/Plugins/Online/Experimental/OnlineServicesNull/Source/OnlineServicesNull.Build.cs b/Engine/Plugins/Online/Experimental/OnlineServicesNull/Source/OnlineServicesNull.Build.cs
--- a/Engine/Plugins/Online/Experimental/OnlineServicesNull/Source/OnlineServicesNull.Build.cs
+++ b/Engine/Plugins/Online/Experimental/OnlineServicesNull/Source/OnlineServicesNull.Build.cs
@@ -13,7 +13,8 @@ public class OnlineServicesNull : ModuleRules
 				"CoreOnline",
 				"Sockets",
 				"OnlineServicesInterface",
-				"OnlineServicesCommon"
+				"OnlineServicesCommon",
+				"OnlineSubsystem"
 			}
 		);
 
@@ -22,7 +23,6 @@ public class OnlineServicesNull : ModuleRules
 			PublicDependencyModuleNames.AddRange(
 				new string[]
 				{
-					"OnlineSubsystem",
 					"OnlineSubsystemUtils"
 				}
 			);
diff --git a/Engine/Plugins/Online/Experimental/OnlineServicesNull/Source/Private/Online/LobbiesNull.cpp b/Engine/Plugins/Online/Experimental/OnlineServicesNull/Source/Private/Online/LobbiesNull.cpp
--- a/Engine/Plugins/Online/Experimental/OnlineServicesNull/Source/Private/Online/LobbiesNull.cpp
+++ b/Engine/Plugins/Online/Experimental/OnlineServicesNull/Source/Private/Online/LobbiesNull.cpp
@@ -1,6 +1,5 @@
 // Copyright Epic Games, Inc. All Rights Reserved.
 
-#if WITH_ENGINE
 #include "Online/LobbiesNull.h"
 #include "Online/OnlineServicesNull.h"
 #include "Misc/Guid.h"
@@ -8,8 +7,9 @@
 #include "Online/NboSerializerNullSvc.h"
 #include "UObject/CoreNet.h"
 #include "Online/LobbyRegistryNull.h"
-#include "OnlineSubsystemTypes.h" //todo: preferably shouldnt need this
-#include "OnlineSubsystemUtils.h" //todo: preferably shouldnt need this
+#if WITH_ENGINE
+	#include "OnlineSubsystemUtils.h" // Needed for GetPortFromNetDriver
+#endif //WITH_ENGINE
 
 namespace UE::Online { 
 
@@ -411,5 +411,3 @@ TOnlineResult<FGetJoinedLobbies> FLobbiesNull::GetJoinedLobbies(FGetJoinedLobbie
 
 
 #undef NOT_IMPLEMENTED
-
-#endif // WITH_ENGINE
diff --git a/Engine/Plugins/Online/Experimental/OnlineServicesNull/Source/Public/Online/LobbiesNull.h b/Engine/Plugins/Online/Experimental/OnlineServicesNull/Source/Public/Online/LobbiesNull.h
--- a/Engine/Plugins/Online/Experimental/OnlineServicesNull/Source/Public/Online/LobbiesNull.h
+++ b/Engine/Plugins/Online/Experimental/OnlineServicesNull/Source/Public/Online/LobbiesNull.h
@@ -1,8 +1,6 @@
 // Copyright Epic Games, Inc. All Rights Reserved.
 #pragma once
 
-#if WITH_ENGINE
-
 #include "CoreMinimal.h"
 #include "Online/LobbiesCommon.h"
 #include "Online/Lobbies.h"
@@ -153,5 +151,3 @@ private:
 };
 
 } // namespace UE::Online
-
-#endif // WITH_ENGINE
