OsvcNull's Presence interface along with a few changes that came up as a result of trying to get it to work

(Note: Originally intended on using the presence common default impl for null and just stub-notimplemented every method, but now instead opting for a minorly functioning presence interface, atleast one that passes our basic suite of tests)

diff --git a/Engine/Plugins/Online/Experimental/OnlineServices/Source/OnlineServicesCommon/Private/Online/PresenceCommon.cpp b/Engine/Plugins/Online/Experimental/OnlineServices/Source/OnlineServicesCommon/Private/Online/PresenceCommon.cpp
--- a/Engine/Plugins/Online/Experimental/OnlineServices/Source/OnlineServicesCommon/Private/Online/PresenceCommon.cpp
+++ b/Engine/Plugins/Online/Experimental/OnlineServices/Source/OnlineServicesCommon/Private/Online/PresenceCommon.cpp
@@ -54,12 +54,6 @@ public:
 				{
 					// push results
 					Result.Presences.Add(NewOp.GetOkValue().Presence);
-
-					// final calculations
-					if (--NumPresencesRemaining < 0)
-					{
-						Promise.SetValue(FBatchQueryPresenceHelperResult{MoveTemp(Result), MoveTemp(ErrorState)});
-					}
 				}
 				else
 				{
@@ -67,6 +61,13 @@ public:
 					// we're only reporting the most recent one, more verbose logging on each individual error can be found in the logs
 					ErrorState = NewOp.GetErrorValue();
 				}
+
+
+				// final calculations
+				if (--NumPresencesRemaining <= 0)
+				{
+					Promise.SetValue(FBatchQueryPresenceHelperResult{ MoveTemp(Result), MoveTemp(ErrorState) });
+				}
 			});
 		}
 
@@ -86,11 +87,11 @@ TOnlineAsyncOpHandle<FBatchQueryPresence> FPresenceCommon::BatchQueryPresence(FB
 {
 	TOnlineAsyncOpRef<FBatchQueryPresence> Op = GetOp<FBatchQueryPresence>(MoveTemp(Params));
 	TSharedRef<FBatchQueryPresenceHelper> Helper = MakeShared<FBatchQueryPresenceHelper>(Op, this->AsShared());
+
 	Op->Then([this, Helper](TOnlineAsyncOp<FBatchQueryPresence>& InAsyncOp) mutable
 	{
 		return Helper->GetPromise();
 	})
-	
 	.Then([this, Helper](TOnlineAsyncOp<FBatchQueryPresence>& InAsyncOp, FBatchQueryPresenceHelperResult Result) mutable
 	{
 		if (Result.ErrorState == Errors::Success())
@@ -101,7 +102,9 @@ TOnlineAsyncOpHandle<FBatchQueryPresence> FPresenceCommon::BatchQueryPresence(FB
 		{
 			InAsyncOp.SetError(MoveTemp(Result.ErrorState));
 		}
-	});
+	})
+	.Enqueue(Services.GetParallelQueue());
+
 	return Op->GetHandle();
 }
 
diff --git a/Engine/Plugins/Online/Experimental/OnlineServices/Source/OnlineServicesInterface/Public/Online/Presence.h b/Engine/Plugins/Online/Experimental/OnlineServices/Source/OnlineServicesInterface/Public/Online/Presence.h
--- a/Engine/Plugins/Online/Experimental/OnlineServices/Source/OnlineServicesInterface/Public/Online/Presence.h
+++ b/Engine/Plugins/Online/Experimental/OnlineServices/Source/OnlineServicesInterface/Public/Online/Presence.h
@@ -243,6 +243,46 @@ inline FPartialUpdatePresence::Params::FMutations& FPartialUpdatePresence::Param
 	return *this;
 }
 
+inline TSharedRef<FUserPresence> ApplyPresenceMutations(const FUserPresence& BasePresence, const FPartialUpdatePresence::Params::FMutations& Mutations)
+{
+	TSharedRef<FUserPresence> NewPresence = MakeShared<FUserPresence>();
+	*NewPresence = BasePresence;
+
+	if (Mutations.Status.IsSet())
+	{
+		NewPresence->Status = Mutations.Status.GetValue();
+	}
+
+	if (Mutations.Joinability.IsSet())
+	{
+		NewPresence->Joinability = Mutations.Joinability.GetValue();
+	}
+
+	if (Mutations.GameStatus.IsSet())
+	{
+		NewPresence->GameStatus = Mutations.GameStatus.GetValue();
+	}
+
+	if (Mutations.StatusString.IsSet())
+	{
+		NewPresence->StatusString = Mutations.StatusString.GetValue();
+	}
+
+	// Insert any updated keys / remove from removed
+	for (const TPair<FString, FPresenceVariant>& NewUpdatedProperty : Mutations.UpdatedProperties)
+	{
+		NewPresence->Properties.Add(NewUpdatedProperty.Key, NewUpdatedProperty.Value);
+	}
+
+	// Merge any removed keys / remove from updated
+	for (const FString& NewRemovedKey : Mutations.RemovedProperties)
+	{
+		NewPresence->Properties.Remove(NewRemovedKey);
+	}
+
+	return NewPresence;
+}
+
 /** Struct for PresenceUpdated event */
 struct FPresenceUpdated
 {
diff --git a/Engine/Plugins/Online/Experimental/OnlineServicesNull/Source/Private/Online/OnlineServicesNull.cpp b/Engine/Plugins/Online/Experimental/OnlineServicesNull/Source/Private/Online/OnlineServicesNull.cpp
--- a/Engine/Plugins/Online/Experimental/OnlineServicesNull/Source/Private/Online/OnlineServicesNull.cpp
+++ b/Engine/Plugins/Online/Experimental/OnlineServicesNull/Source/Private/Online/OnlineServicesNull.cpp
@@ -7,6 +7,7 @@
 #include "Online/AuthNull.h"
 #include "Online/StatsNull.h"
 #include "Online/LobbiesNull.h"
+#include "Online/PresenceNull.h"
 #include "Online/TitleFileNull.h"
 
 namespace UE::Online {
@@ -35,6 +36,7 @@ void FOnlineServicesNull::RegisterComponents()
 {
 	Components.Register<FAchievementsNull>(*this);
 	Components.Register<FAuthNull>(*this);
+	Components.Register<FPresenceNull>(*this);
 	Components.Register<FStatsNull>(*this);
 #if WITH_ENGINE
 	Components.Register<FLobbiesNull>(*this);
diff --git a/Engine/Plugins/Online/Experimental/OnlineServicesNull/Source/Private/Online/PresenceNull.cpp b/Engine/Plugins/Online/Experimental/OnlineServicesNull/Source/Private/Online/PresenceNull.cpp
new file mode 100644
--- /dev/null
+++ b/Engine/Plugins/Online/Experimental/OnlineServicesNull/Source/Private/Online/PresenceNull.cpp
@@ -0,0 +1,124 @@
+// Copyright Epic Games, Inc. All Rights Reserved.
+
+#include "Online/PresenceNull.h"
+#include "Online/OnlineAsyncOp.h"
+#include "Online/OnlineServicesCommon.h"
+#include "Online/OnlineServicesNull.h"
+
+namespace UE::Online {
+
+FPresenceNull::FPresenceNull(FOnlineServicesNull& InServices)
+	: FPresenceCommon(InServices)
+{
+
+}
+
+void FPresenceNull::Initialize()
+{
+
+}
+
+void FPresenceNull::PreShutdown()
+{
+
+}
+
+TOnlineAsyncOpHandle<FQueryPresence> FPresenceNull::QueryPresence(FQueryPresence::Params&& Params)
+{
+
+	TOnlineAsyncOpRef<FQueryPresence> Op = GetOp<FQueryPresence>(MoveTemp(Params));
+
+	Op->Then([this](TOnlineAsyncOp<FQueryPresence>& InAsyncOp) mutable
+	{
+		const FQueryPresence::Params& Params = InAsyncOp.GetParams();
+
+		if(Params.bListenToChanges)
+		{
+			PresenceListeners.FindOrAdd(Params.LocalUserId).Add(Params.TargetUserId);
+		}
+
+		if(Presences.Contains(Params.TargetUserId))
+		{
+			InAsyncOp.SetResult({Presences.FindChecked(Params.TargetUserId)});
+		}
+		else
+		{
+			InAsyncOp.SetError(Errors::NotFound());
+		}
+	})
+    .Enqueue(GetSerialQueue());
+
+	return Op->GetHandle();
+}
+
+TOnlineResult<FGetCachedPresence> FPresenceNull::GetCachedPresence(FGetCachedPresence::Params&& Params)
+{
+	if (Presences.Contains(Params.LocalUserId))
+	{
+		return TOnlineResult<FGetCachedPresence>({Presences.FindChecked(Params.LocalUserId)});
+	}
+	else
+	{
+		return TOnlineResult<FGetCachedPresence>(Errors::NotFound());
+	}
+}
+
+TOnlineAsyncOpHandle<FUpdatePresence> FPresenceNull::UpdatePresence(FUpdatePresence::Params&& Params)
+{
+
+	TOnlineAsyncOpRef<FUpdatePresence> Op = GetOp<FUpdatePresence>(MoveTemp(Params));
+
+	Op->Then([this](TOnlineAsyncOp<FUpdatePresence>& InAsyncOp) mutable
+	{
+		const FUpdatePresence::Params& Params = InAsyncOp.GetParams();
+		Presences.Add(Params.LocalUserId, Params.Presence);
+
+		for (const TPair<FOnlineAccountIdHandle, TSet<FOnlineAccountIdHandle>>& Pairs : PresenceListeners)
+		{
+			if (Pairs.Value.Contains(Params.LocalUserId))
+			{
+				OnPresenceUpdatedEvent.Broadcast({Pairs.Key, Params.Presence});
+			}
+		}
+
+		InAsyncOp.SetResult({});
+	})
+    .Enqueue(GetSerialQueue());
+
+	return Op->GetHandle();
+}
+
+TOnlineAsyncOpHandle<FPartialUpdatePresence> FPresenceNull::PartialUpdatePresence(FPartialUpdatePresence::Params&& Params)
+{
+
+	TOnlineAsyncOpRef<FPartialUpdatePresence> Op = GetOp<FPartialUpdatePresence>(MoveTemp(Params));
+
+	Op->Then([this](TOnlineAsyncOp<FPartialUpdatePresence>& InAsyncOp) mutable
+	{
+		TSharedRef<FUserPresence> NewPresence = MakeShared<FUserPresence>();
+		const FPartialUpdatePresence::Params::FMutations& Mutations = InAsyncOp.GetParams().Mutations;
+		
+		if (Presences.Contains(InAsyncOp.GetParams().LocalUserId))
+		{
+			*NewPresence = *Presences.FindChecked(InAsyncOp.GetParams().LocalUserId);
+		}
+
+		TSharedRef<FUserPresence> MutatedPresence = ApplyPresenceMutations(*NewPresence, Mutations);
+		Presences.Add(InAsyncOp.GetParams().LocalUserId, MutatedPresence);
+
+		for (const TPair<FOnlineAccountIdHandle, TSet<FOnlineAccountIdHandle>>& Pairs : PresenceListeners)
+		{
+			if (Pairs.Value.Contains(InAsyncOp.GetParams().LocalUserId))
+			{
+				OnPresenceUpdatedEvent.Broadcast({ Pairs.Key, MutatedPresence });
+			}
+		}
+
+		InAsyncOp.SetResult({});
+	})
+    .Enqueue(GetSerialQueue());
+
+	return Op->GetHandle();
+}
+
+} //namespace UE::Online
\ No newline at end of file
diff --git a/Engine/Plugins/Online/Experimental/OnlineServicesNull/Source/Public/Online/PresenceNull.h b/Engine/Plugins/Online/Experimental/OnlineServicesNull/Source/Public/Online/PresenceNull.h
new file mode 100644
--- /dev/null
+++ b/Engine/Plugins/Online/Experimental/OnlineServicesNull/Source/Public/Online/PresenceNull.h
@@ -0,0 +1,31 @@
+// Copyright Epic Games, Inc. All Rights Reserved.
+
+#pragma once
+
+#include "CoreMinimal.h"
+#include "Online/PresenceCommon.h"
+
+namespace UE::Online {
+
+class FOnlineServicesNull;
+
+class FPresenceNull : public FPresenceCommon
+{
+public:
+	FPresenceNull(FOnlineServicesNull& InServices);
+
+	virtual void Initialize() override;
+	virtual void PreShutdown() override;
+
+	virtual TOnlineAsyncOpHandle<FQueryPresence> QueryPresence(FQueryPresence::Params&& Params) override;
+	virtual TOnlineResult<FGetCachedPresence> GetCachedPresence(FGetCachedPresence::Params&& Params) override;
+	virtual TOnlineAsyncOpHandle<FUpdatePresence> UpdatePresence(FUpdatePresence::Params&& Params) override;
+	virtual TOnlineAsyncOpHandle<FPartialUpdatePresence> PartialUpdatePresence(FPartialUpdatePresence::Params&& Params) override;
+
+protected:
+	TMap<FOnlineAccountIdHandle, TSharedRef<const FUserPresence>> Presences;
+	TMap<FOnlineAccountIdHandle, TSet<FOnlineAccountIdHandle>> PresenceListeners;
+};
+
+	/* UE::Online */
+}
