- Implement User File Null

diff --git a/Engine/Plugins/Online/Experimental/OnlineServicesNull/Source/Private/Online/OnlineServicesNull.cpp b/Engine/Plugins/Online/Experimental/OnlineServicesNull/Source/Private/Online/OnlineServicesNull.cpp
--- a/Engine/Plugins/Online/Experimental/OnlineServicesNull/Source/Private/Online/OnlineServicesNull.cpp
+++ b/Engine/Plugins/Online/Experimental/OnlineServicesNull/Source/Private/Online/OnlineServicesNull.cpp
@@ -9,6 +9,7 @@
 #include "Online/LobbiesNull.h"
 #include "Online/PresenceNull.h"
 #include "Online/TitleFileNull.h"
+#include "Online/UserFileNull.h"
 
 namespace UE::Online {
 
@@ -42,6 +43,7 @@ void FOnlineServicesNull::RegisterComponents()
 	Components.Register<FLobbiesNull>(*this);
 #endif
 	Components.Register<FTitleFileNull>(*this);
+	Components.Register<FUserFileNull>(*this);
 	FOnlineServicesCommon::RegisterComponents();
 }
 
diff --git a/Engine/Plugins/Online/Experimental/OnlineServicesNull/Source/Private/Online/UserFileNull.cpp b/Engine/Plugins/Online/Experimental/OnlineServicesNull/Source/Private/Online/UserFileNull.cpp
new file mode 100644
--- /dev/null
+++ b/Engine/Plugins/Online/Experimental/OnlineServicesNull/Source/Private/Online/UserFileNull.cpp
@@ -0,0 +1,228 @@
+// Copyright Epic Games, Inc. All Rights Reserved.
+
+#include "Online/UserFileNull.h"
+
+#include "Containers/StringConv.h"
+#include "Online/AuthNull.h"
+#include "Online/OnlineErrorDefinitions.h"
+#include "Online/OnlineServicesNull.h"
+#include "Online/OnlineServicesNullTypes.h"
+
+namespace UE::Online {
+
+FUserFileNull::FUserFileNull(FOnlineServicesNull& InOwningSubsystem)
+	: Super(InOwningSubsystem)
+{
+}
+
+TOnlineAsyncOpHandle<FUserFileEnumerateFiles> FUserFileNull::EnumerateFiles(FUserFileEnumerateFiles::Params&& InParams)
+{
+	TOnlineAsyncOpRef<FUserFileEnumerateFiles> Op = GetOp<FUserFileEnumerateFiles>(MoveTemp(InParams));
+	const FUserFileEnumerateFiles::Params& Params = Op->GetParams();
+
+	if (!Services.Get<FAuthNull>()->IsLoggedIn(Params.LocalUserId))
+	{
+		Op->SetError(Errors::InvalidUser());
+		return Op->GetHandle();
+	}
+
+	if (!UserToFilesMap.Contains(Params.LocalUserId))
+	{
+		if (!InitialFileStateFromConfig.IsSet())
+		{
+			LoadUserFilesFromConfig();
+		}
+		UserToFilesMap.Emplace(Params.LocalUserId, *InitialFileStateFromConfig);
+	}
+
+	Op->SetResult({});
+	return Op->GetHandle();
+}
+
+TOnlineResult<FUserFileGetEnumeratedFiles> FUserFileNull::GetEnumeratedFiles(FUserFileGetEnumeratedFiles::Params&& Params)
+{
+	if (!Services.Get<FAuthNull>()->IsLoggedIn(Params.LocalUserId))
+	{
+		return TOnlineResult<FUserFileGetEnumeratedFiles>(Errors::InvalidUser());
+	}
+
+	FUserFileMap* UserFiles = UserToFilesMap.Find(Params.LocalUserId);
+	if (!UserFiles)
+	{
+		// Need to call EnumerateFiles first.
+		return TOnlineResult<FUserFileGetEnumeratedFiles>(Errors::InvalidState());
+	}
+
+	FUserFileGetEnumeratedFiles::Result Result;
+	UserFiles->GenerateKeyArray(Result.Filenames);
+	return TOnlineResult<FUserFileGetEnumeratedFiles>(MoveTemp(Result));
+}
+
+TOnlineAsyncOpHandle<FUserFileReadFile> FUserFileNull::ReadFile(FUserFileReadFile::Params&& InParams)
+{
+	TOnlineAsyncOpRef<FUserFileReadFile> Op = GetOp<FUserFileReadFile>(MoveTemp(InParams));
+	const FUserFileReadFile::Params& Params = Op->GetParams();
+
+	if (!Services.Get<FAuthNull>()->IsLoggedIn(Params.LocalUserId))
+	{
+		Op->SetError(Errors::InvalidUser());
+		return Op->GetHandle();
+	}
+
+	FUserFileMap* UserFiles = UserToFilesMap.Find(Params.LocalUserId);
+	if (!UserFiles)
+	{
+		// Need to call EnumerateFiles first.
+		Op->SetError(Errors::InvalidState());
+		return Op->GetHandle();
+	}
+
+	if (Params.Filename.IsEmpty())
+	{
+		Op->SetError(Errors::InvalidParams());
+		return Op->GetHandle();
+	}
+
+	const FUserFileContentsRef* Found = UserFiles->Find(Params.Filename);
+	if (!Found)
+	{
+		Op->SetError(Errors::NotFound());
+		return Op->GetHandle();
+	}
+
+	Op->SetResult({*Found});
+	return Op->GetHandle();
+}
+
+TOnlineAsyncOpHandle<FUserFileWriteFile> FUserFileNull::WriteFile(FUserFileWriteFile::Params&& InParams)
+{
+	TOnlineAsyncOpRef<FUserFileWriteFile> Op = GetOp<FUserFileWriteFile>(MoveTemp(InParams));
+	const FUserFileWriteFile::Params& Params = Op->GetParams();
+
+	if (!Services.Get<FAuthNull>()->IsLoggedIn(Params.LocalUserId))
+	{
+		Op->SetError(Errors::InvalidUser());
+		return Op->GetHandle();
+	}
+
+	FUserFileMap* UserFiles = UserToFilesMap.Find(Params.LocalUserId);
+	if (!UserFiles)
+	{
+		// Need to call EnumerateFiles first.
+		Op->SetError(Errors::InvalidState());
+		return Op->GetHandle();
+	}
+
+	if (Params.Filename.IsEmpty() || Params.FileContents.IsEmpty())
+	{
+		Op->SetError(Errors::InvalidParams());
+		return Op->GetHandle();
+	}
+
+	UserFiles->Emplace(Params.Filename, MakeShared<FUserFileContents>(Params.FileContents));
+	Op->SetResult({});
+	return Op->GetHandle();
+}
+
+TOnlineAsyncOpHandle<FUserFileCopyFile> FUserFileNull::CopyFile(FUserFileCopyFile::Params&& InParams)
+{
+	TOnlineAsyncOpRef<FUserFileCopyFile> Op = GetOp<FUserFileCopyFile>(MoveTemp(InParams));
+	const FUserFileCopyFile::Params& Params = Op->GetParams();
+
+	if (!Services.Get<FAuthNull>()->IsLoggedIn(Params.LocalUserId))
+	{
+		Op->SetError(Errors::InvalidUser());
+		return Op->GetHandle();
+	}
+
+	FUserFileMap* UserFiles = UserToFilesMap.Find(Params.LocalUserId);
+	if (!UserFiles)
+	{
+		// Need to call EnumerateFiles first.
+		Op->SetError(Errors::InvalidState());
+		return Op->GetHandle();
+	}
+
+	if (Params.SourceFilename.IsEmpty() || Params.TargetFilename.IsEmpty())
+	{
+		Op->SetError(Errors::InvalidParams());
+		return Op->GetHandle();
+	}
+
+	const FUserFileContentsRef* SourceFile = UserFiles->Find(Params.SourceFilename);
+	if (!SourceFile)
+	{
+		Op->SetError(Errors::NotFound());
+		return Op->GetHandle();
+	}
+
+	UserFiles->Emplace(Params.TargetFilename, *SourceFile);
+	Op->SetResult({});
+	return Op->GetHandle();
+}
+
+TOnlineAsyncOpHandle<FUserFileDeleteFile> FUserFileNull::DeleteFile(FUserFileDeleteFile::Params&& InParams)
+{
+	TOnlineAsyncOpRef<FUserFileDeleteFile> Op = GetOp<FUserFileDeleteFile>(MoveTemp(InParams));
+	const FUserFileDeleteFile::Params& Params = Op->GetParams();
+
+	if (!Services.Get<FAuthNull>()->IsLoggedIn(Params.LocalUserId))
+	{
+		Op->SetError(Errors::InvalidUser());
+		return Op->GetHandle();
+	}
+
+	FUserFileMap* UserFiles = UserToFilesMap.Find(Params.LocalUserId);
+	if (!UserFiles)
+	{
+		// Need to call EnumerateFiles first.
+		Op->SetError(Errors::InvalidState());
+		return Op->GetHandle();
+	}
+
+	if (Params.Filename.IsEmpty())
+	{
+		Op->SetError(Errors::InvalidParams());
+		return Op->GetHandle();
+	}
+
+	const bool bRemoved = UserFiles->Remove(Params.Filename) != 0;
+	if (!bRemoved)
+	{
+		Op->SetError(Errors::NotFound());
+		return Op->GetHandle();
+	}
+
+	Op->SetResult({});
+	return Op->GetHandle();
+}
+
+void FUserFileNull::LoadUserFilesFromConfig()
+{
+	const TCHAR* ConfigSection = TEXT("OnlineServices.Null.UserFile");
+
+	InitialFileStateFromConfig.Emplace();
+
+	for (int FileIdx = 0;; FileIdx++)
+	{
+		FString Filename;
+		GConfig->GetString(ConfigSection, *FString::Printf(TEXT("File_%d_Name"), FileIdx), Filename, GEngineIni);
+		if (Filename.IsEmpty())
+		{
+			break;
+		}
+
+		FString FileContentsStr;
+		GConfig->GetString(ConfigSection, *FString::Printf(TEXT("File_%d_Contents"), FileIdx), FileContentsStr, GEngineIni);
+
+		if (!FileContentsStr.IsEmpty())
+		{
+			const int ArrayLen = FileContentsStr.GetCharArray().Num() * sizeof(TCHAR);
+			FUserFileContents FileContents((uint8*)FileContentsStr.GetCharArray().GetData(), ArrayLen);
+
+			InitialFileStateFromConfig->Emplace(MoveTemp(Filename), MakeShared<const FUserFileContents>(MoveTemp(FileContents)));
+		}
+	}
+}
+
+/* UE::Online */ }
diff --git a/Engine/Plugins/Online/Experimental/OnlineServicesNull/Source/Public/Online/UserFileNull.h b/Engine/Plugins/Online/Experimental/OnlineServicesNull/Source/Public/Online/UserFileNull.h
new file mode 100644
--- /dev/null
+++ b/Engine/Plugins/Online/Experimental/OnlineServicesNull/Source/Public/Online/UserFileNull.h
@@ -0,0 +1,36 @@
+// Copyright Epic Games, Inc. All Rights Reserved.
+
+#pragma once
+
+#include "Misc/Optional.h"
+#include "Online/UserFileCommon.h"
+
+namespace UE::Online {
+
+class FOnlineServicesNull;
+
+class ONLINESERVICESNULL_API FUserFileNull : public FUserFileCommon
+{
+public:
+	using Super = FUserFileCommon;
+
+	FUserFileNull(FOnlineServicesNull& InOwningSubsystem);
+
+	// IUserFile
+	virtual TOnlineAsyncOpHandle<FUserFileEnumerateFiles> EnumerateFiles(FUserFileEnumerateFiles::Params&& Params) override;
+	virtual TOnlineResult<FUserFileGetEnumeratedFiles> GetEnumeratedFiles(FUserFileGetEnumeratedFiles::Params&& Params) override;
+	virtual TOnlineAsyncOpHandle<FUserFileReadFile> ReadFile(FUserFileReadFile::Params&& Params) override;
+	virtual TOnlineAsyncOpHandle<FUserFileWriteFile> WriteFile(FUserFileWriteFile::Params&& Params) override;
+	virtual TOnlineAsyncOpHandle<FUserFileCopyFile> CopyFile(FUserFileCopyFile::Params&& Params) override;
+	virtual TOnlineAsyncOpHandle<FUserFileDeleteFile> DeleteFile(FUserFileDeleteFile::Params&& Params) override;
+
+protected:
+	using FUserFileMap = TMap<FString, FUserFileContentsRef>;
+
+	TMap<FOnlineAccountIdHandle, FUserFileMap> UserToFilesMap;
+
+	TOptional<FUserFileMap> InitialFileStateFromConfig;
+	void LoadUserFilesFromConfig();
+};
+
+/* UE::Online */ }
