OSSv2 Sessions Interface - Null Implementation

diff --git a/Engine/Plugins/Online/Experimental/OnlineServices/Source/OnlineServicesInterface/Private/Online/Sessions.cpp b/Engine/Plugins/Online/Experimental/OnlineServices/Source/OnlineServicesInterface/Private/Online/Sessions.cpp
--- a/Engine/Plugins/Online/Experimental/OnlineServices/Source/OnlineServicesInterface/Private/Online/Sessions.cpp
+++ b/Engine/Plugins/Online/Experimental/OnlineServices/Source/OnlineServicesInterface/Private/Online/Sessions.cpp
@@ -4,6 +4,58 @@
 
 namespace UE::Online {
 
+const TCHAR* LexToString(ESessionsComparisonOp Comparison)
+{
+	switch (Comparison)
+	{
+	case ESessionsComparisonOp::Equals:				return TEXT("Equals");
+	case ESessionsComparisonOp::NotEquals:			return TEXT("NotEquals");
+	case ESessionsComparisonOp::GreaterThan:		return TEXT("GreaterThan");
+	case ESessionsComparisonOp::GreaterThanEquals:	return TEXT("GreaterThanEquals");
+	case ESessionsComparisonOp::LessThan:			return TEXT("LessThan");
+	case ESessionsComparisonOp::LessThanEquals:		return TEXT("LessThanEquals");
+	default:										checkNoEntry(); // Intentional fallthrough
+	case ESessionsComparisonOp::Near:				return TEXT("Near");
+	}
+}
+
+void LexFromString(ESessionsComparisonOp& OutComparison, const TCHAR* InStr)
+{
+	if (FCString::Stricmp(InStr, TEXT("Equals")) == 0)
+	{
+		OutComparison = ESessionsComparisonOp::Equals;
+	}
+	else if (FCString::Stricmp(InStr, TEXT("NotEquals")) == 0)
+	{
+		OutComparison = ESessionsComparisonOp::NotEquals;
+	}
+	else if (FCString::Stricmp(InStr, TEXT("GreaterThan")) == 0)
+	{
+		OutComparison = ESessionsComparisonOp::GreaterThan;
+	}
+	else if (FCString::Stricmp(InStr, TEXT("GreaterThanEquals")) == 0)
+	{
+		OutComparison = ESessionsComparisonOp::GreaterThanEquals;
+	}
+	else if (FCString::Stricmp(InStr, TEXT("LessThan")) == 0)
+	{
+		OutComparison = ESessionsComparisonOp::LessThan;
+	}
+	else if (FCString::Stricmp(InStr, TEXT("LessThanEquals")) == 0)
+	{
+		OutComparison = ESessionsComparisonOp::LessThanEquals;
+	}
+	else if (FCString::Stricmp(InStr, TEXT("Near")) == 0)
+	{
+		OutComparison = ESessionsComparisonOp::Near;
+	}
+	else
+	{
+		checkNoEntry();
+		OutComparison = ESessionsComparisonOp::Near;
+	}
+}
+
 const TCHAR* LexToString(ECustomSessionSettingVisibility Value)
 {
 	switch (Value)
@@ -253,4 +305,18 @@
 	return *this;
 }
 
+FSession::FSession()
+{
+
+}
+
+FSession::FSession(const FSession& InSession)
+	: OwnerUserId(InSession.OwnerUserId)
+	, SessionId(InSession.SessionId)
+	, CurrentState(InSession.CurrentState)
+	, SessionSettings(InSession.SessionSettings)
+{
+
+}
+
 /* UE::Online */ }
diff --git a/Engine/Plugins/Online/Experimental/OnlineServices/Source/OnlineServicesInterface/Public/Online/Sessions.h b/Engine/Plugins/Online/Experimental/OnlineServices/Source/OnlineServicesInterface/Public/Online/Sessions.h
--- a/Engine/Plugins/Online/Experimental/OnlineServices/Source/OnlineServicesInterface/Public/Online/Sessions.h
+++ b/Engine/Plugins/Online/Experimental/OnlineServices/Source/OnlineServicesInterface/Public/Online/Sessions.h
@@ -9,9 +9,33 @@
 
 namespace UE::Online {
 
+enum class ESessionsComparisonOp : uint8
+{
+	Equals,
+	NotEquals,
+	GreaterThan,
+	GreaterThanEquals,
+	LessThan,
+	LessThanEquals,
+	Near
+};
+ONLINESERVICESINTERFACE_API const TCHAR* LexToString(ESessionsComparisonOp EnumVal);
+ONLINESERVICESINTERFACE_API void LexFromString(ESessionsComparisonOp& OutComparison, const TCHAR* InStr);
+
 struct FFindSessionsSearchFilter
 {
-	// TODO: Will share this type with Lobbies, so the implementation will be identical for now
+	// TODO: Will change after SchemaVariant work
+
+	using FVariantType = TVariant<FString, int64, double, bool>;
+
+	/** Name of the custom setting to be used as filter */
+	FName Key;
+
+	/** The type of comparison to perform */
+	ESessionsComparisonOp ComparisonOp;
+
+	/** Value to use when comparing the filter */
+	FVariantType Value;
 };
 
 enum class ECustomSessionSettingVisibility : uint8
@@ -177,8 +201,11 @@
 	FSessionSettingsUpdate& operator+=(FSessionSettingsUpdate&& UpdatedValue);
 };
 
-struct FSession
+struct ONLINESERVICESINTERFACE_API FSession
 {
+	FSession();
+	FSession(const FSession& InSession);
+
 	/** The user who currently owns the session */
 	FOnlineAccountIdHandle OwnerUserId;
 
@@ -215,7 +242,7 @@
 
 	struct Result
 	{
-		TArray<TSharedRef<FSession>> Sessions;
+		TArray<TSharedRef<const FSession>> Sessions;
 	};
 };
 
@@ -230,7 +257,9 @@
 
 	struct Result
 	{
-		TSharedRef<FSession> Session;
+		TSharedRef<const FSession> Session;
+
+		Result() = delete; // cannot default construct due to TSharedRef
 	};
 };
 
@@ -245,7 +274,9 @@
 
 	struct Result
 	{
-		TSharedRef<FSession> Session;
+		TSharedRef<const FSession> Session;
+
+		Result() = delete; // cannot default construct due to TSharedRef
 	};
 };
 
@@ -264,13 +295,15 @@
 		/** Settings object to define session properties during creation */
 		FSessionSettings SessionSettings;
 
-		/** Local users who will be joining the session on successful creation, along with the session creator */
+		/** Information for all local users who will join the session (includes the session creator) */
 		FSessionMembersMap LocalUsers;
 	};
 
 	struct Result
 	{
 		TSharedRef<const FSession> Session;
+
+		Result() = delete; // cannot default construct due to TSharedRef
 	};
 };
 
@@ -293,6 +326,8 @@
 	struct Result
 	{
 		TSharedRef<const FSession> Session;
+
+		Result() = delete; // cannot default construct due to TSharedRef
 	};
 };
 
@@ -308,13 +343,16 @@
 		/* The local name for the session. */
 		FName SessionName;
 
+		/** Ids for local users (other than the main caller) who will also be leaving the session */
+		TArray<FOnlineAccountIdHandle> AdditionalLocalUsers;
+
 		/* Whether the call should attempt to destroy the session instead of just leave it */
 		bool bDestroySession;
 	};
 
 	struct Result
 	{
-		TSharedRef<const FSession> Session;
+
 	};
 };
 
@@ -352,7 +390,10 @@
 
 	struct Params
 	{
-		/* Local users who will be joining the session */
+		/** The local user agent which will perform the action. */
+		FOnlineAccountIdHandle LocalUserId;
+
+		/* Information for all local users who will join the session (includes the session creator) */
 		FSessionMembersMap LocalUsers;
 
 		/* Local name for the session */
@@ -368,6 +409,8 @@
 	struct Result
 	{
 		TSharedRef<const FSession> Session;
+
+		Result() = delete; // cannot default construct due to TSharedRef
 	};
 };
 
@@ -384,15 +427,17 @@
 		FName SessionName;
 
 		/* A reference to the session to be joined. */
-		TSharedRef<FSession> Session;
+		TSharedRef<const FSession> Session;
 
-		/* Local users who will be joining the session. */
+		/* Information for all local users who will join the session (includes the session creator)*/
 		FSessionMembersMap LocalUsers;
 	};
 
 	struct Result
 	{
 		TSharedRef<const FSession> Session;
+
+		Result() = delete; // cannot default construct due to TSharedRef
 	};
 };
 
@@ -503,7 +548,7 @@
 	TArray<FOnlineAccountIdHandle> LocalUserIds;
 
 	/* A shared reference to the session joined. */
-	TSharedRef<FSession> Session;
+	TSharedRef<const FSession> Session;
 };
 
 struct FSessionLeft
@@ -822,11 +867,11 @@
 BEGIN_ONLINE_STRUCT_META(FLeaveSession::Params)
 	ONLINE_STRUCT_FIELD(FLeaveSession::Params, LocalUserId),
 	ONLINE_STRUCT_FIELD(FLeaveSession::Params, SessionName),
+	ONLINE_STRUCT_FIELD(FLeaveSession::Params, AdditionalLocalUsers),
 	ONLINE_STRUCT_FIELD(FLeaveSession::Params, bDestroySession)
 END_ONLINE_STRUCT_META()
 
 BEGIN_ONLINE_STRUCT_META(FLeaveSession::Result)
-	ONLINE_STRUCT_FIELD(FLeaveSession::Result, Session)
 END_ONLINE_STRUCT_META()
 
 BEGIN_ONLINE_STRUCT_META(FFindSessions::Params)
@@ -842,6 +887,7 @@
 END_ONLINE_STRUCT_META()
 
 BEGIN_ONLINE_STRUCT_META(FStartMatchmaking::Params)
+	ONLINE_STRUCT_FIELD(FStartMatchmaking::Params, LocalUserId),
 	ONLINE_STRUCT_FIELD(FStartMatchmaking::Params, LocalUsers),
 	ONLINE_STRUCT_FIELD(FStartMatchmaking::Params, SessionName),
 	ONLINE_STRUCT_FIELD(FStartMatchmaking::Params, SessionSettings),
diff --git a/Engine/Plugins/Online/Experimental/OnlineServicesNull/Source/Private/Online/OnlineServicesNull.cpp b/Engine/Plugins/Online/Experimental/OnlineServicesNull/Source/Private/Online/OnlineServicesNull.cpp
--- a/Engine/Plugins/Online/Experimental/OnlineServicesNull/Source/Private/Online/OnlineServicesNull.cpp
+++ b/Engine/Plugins/Online/Experimental/OnlineServicesNull/Source/Private/Online/OnlineServicesNull.cpp
@@ -10,6 +10,7 @@
 #include "Online/PresenceNull.h"
 #include "Online/TitleFileNull.h"
 #include "Online/UserFileNull.h"
+#include "Online/SessionsNull.h"
 
 namespace UE::Online {
 
@@ -39,50 +40,87 @@ void FOnlineServicesNull::RegisterComponents()
 	Components.Register<FAuthNull>(*this);
 	Components.Register<FPresenceNull>(*this);
 	Components.Register<FStatsNull>(*this);
-#if WITH_ENGINE
 	Components.Register<FLobbiesNull>(*this);
-#endif
 	Components.Register<FTitleFileNull>(*this);
 	Components.Register<FUserFileNull>(*this);
+	Components.Register<FSessionsNull>(*this);
 	FOnlineServicesCommon::RegisterComponents();
 }
 
 TOnlineResult<FGetResolvedConnectString> FOnlineServicesNull::GetResolvedConnectString(FGetResolvedConnectString::Params&& Params)
 {
-#if WITH_ENGINE
-	ILobbiesPtr LobbiesPtr = GetLobbiesInterface();
-	check(LobbiesPtr);
-
-	TOnlineResult<FGetJoinedLobbies> JoinedLobbies = LobbiesPtr->GetJoinedLobbies({ Params.LocalUserId });
-	if (JoinedLobbies.IsOk())
+	if(Params.LobbyId.IsValid())
 	{
-		for (TSharedRef<const FLobby>& Lobby : JoinedLobbies.GetOkValue().Lobbies)
+		ILobbiesPtr LobbiesPtr = GetLobbiesInterface();
+		check(LobbiesPtr);
+
+		TOnlineResult<FGetJoinedLobbies> JoinedLobbies = LobbiesPtr->GetJoinedLobbies({ Params.LocalUserId });
+		if (JoinedLobbies.IsOk())
 		{
-			if (Lobby->LobbyId == Params.LobbyId)
+			for (TSharedRef<const FLobby>& Lobby : JoinedLobbies.GetOkValue().Lobbies)
 			{
-				FName Tag = FName(TEXT("ConnectAddress")); // todo: global tag
-				if(Lobby->Attributes.Contains(Tag))
+				if (Lobby->LobbyId == Params.LobbyId)
 				{
-					FGetResolvedConnectString::Result Result;
-					Result.ResolvedConnectString = Lobby->Attributes[Tag].GetString();
- 					return TOnlineResult<FGetResolvedConnectString>(Result);
+					FName Tag = FName(TEXT("ConnectAddress")); // todo: global tag
+					if(Lobby->Attributes.Contains(Tag))
+					{
+						FGetResolvedConnectString::Result Result;
+						Result.ResolvedConnectString = Lobby->Attributes[Tag].GetString();
+ 						return TOnlineResult<FGetResolvedConnectString>(Result);
+					}
+					else
+					{
+						continue;
+					}
 				}
-				else
+			}
+
+			// No matching lobby
+			return TOnlineResult<FGetResolvedConnectString>(Errors::InvalidParams());
+		}
+		else
+		{
+			return TOnlineResult<FGetResolvedConnectString>(JoinedLobbies.GetErrorValue());
+		}
+	}
+	else if (Params.SessionId.IsValid())
+	{
+		ISessionsPtr SessionsPtr = GetSessionsInterface();
+		check(SessionsPtr);
+
+		FGetAllSessions::Params GetAllSessionsParams;
+		TOnlineResult<FGetAllSessions> JoinedSessions = SessionsPtr->GetAllSessions(MoveTemp(GetAllSessionsParams));
+		if (JoinedSessions.IsOk())
+		{
+			for (TSharedRef<const FSession>& Session : JoinedSessions.GetOkValue().Sessions)
+			{
+				if (Session->SessionId == Params.SessionId)
 				{
-					continue;
+					if (Session->SessionSettings.CustomSettings.Contains(CONNECT_STRING_TAG))
+					{
+						FGetResolvedConnectString::Result Result;
+						Result.ResolvedConnectString = Session->SessionSettings.CustomSettings[CONNECT_STRING_TAG].Data.Get<FString>();
+						return TOnlineResult<FGetResolvedConnectString>(Result);
+					}
+					else
+					{
+						continue;
+					}
 				}
 			}
+
+			// No matching session
+			return TOnlineResult<FGetResolvedConnectString>(Errors::InvalidParams());
+		}
+		else
+		{
+			return TOnlineResult<FGetResolvedConnectString>(JoinedSessions.GetErrorValue());
 		}
-		// No matching lobby
-		return TOnlineResult<FGetResolvedConnectString>(Errors::InvalidParams());
 	}
 	else
 	{
-		return TOnlineResult<FGetResolvedConnectString>(JoinedLobbies.GetErrorValue());
+		return TOnlineResult<FGetResolvedConnectString>(Errors::InvalidParams());
 	}
-#else
-	return Super::GetResolvedConnectString(MoveTemp(Params));
-#endif
 }
 
 void FOnlineServicesNull::Initialize()
diff --git a/Engine/Plugins/Online/Experimental/OnlineServicesNull/Source/Private/Online/SessionsNull.cpp b/Engine/Plugins/Online/Experimental/OnlineServicesNull/Source/Private/Online/SessionsNull.cpp
new file mode 100644
--- /dev/null
+++ b/Engine/Plugins/Online/Experimental/OnlineServicesNull/Source/Private/Online/SessionsNull.cpp
@@ -0,0 +1,850 @@
+// Copyright Epic Games, Inc. All Rights Reserved.
+
+#include "Online/SessionsNull.h"
+
+#include "Misc/Guid.h"
+#include "Online/NboSerializerNullSvc.h"
+#include "Online/OnlineServicesNull.h"
+#include "SocketSubsystem.h"
+#include "UObject/CoreNet.h"
+#if WITH_ENGINE
+	#include "OnlineSubsystemUtils.h" // Needed for GetPortFromNetDriver
+#endif
+
+namespace UE::Online {
+
+/** FOnlineSessionIdRegistryNull */
+
+FOnlineSessionIdRegistryNull& FOnlineSessionIdRegistryNull::Get()
+{
+	static FOnlineSessionIdRegistryNull Instance;
+	return Instance;
+}
+
+FOnlineSessionIdHandle FOnlineSessionIdRegistryNull::Find(FString SessionId) const
+{
+	if(const FOnlineSessionIdString* Entry = StringToId.Find(SessionId))
+	{
+		return Entry->Handle;
+	}
+	return FOnlineSessionIdHandle();
+}
+
+FOnlineSessionIdHandle FOnlineSessionIdRegistryNull::FindOrAdd(FString SessionId)
+{
+	const FOnlineSessionIdString* Entry = StringToId.Find(SessionId);
+	if (Entry)
+	{
+		return Entry->Handle;
+	}
+
+	FOnlineSessionIdString& NewId = Ids.Emplace_GetRef();
+	NewId.Data = SessionId;
+	NewId.Handle = FOnlineSessionIdHandle(EOnlineServices::Null, Ids.Num());
+
+	StringToId.Add(SessionId, NewId);
+
+	return NewId.Handle;
+}
+
+FOnlineSessionIdHandle FOnlineSessionIdRegistryNull::GetNext()
+{
+	return FindOrAdd(FGuid().ToString());
+}
+
+const FOnlineSessionIdString* FOnlineSessionIdRegistryNull::GetInternal(const FOnlineSessionIdHandle& Handle) const
+{
+	if (Handle.IsValid() && Handle.GetOnlineServicesType() == EOnlineServices::Null && Handle.GetHandle() <= (uint32)Ids.Num())
+	{
+		return &Ids[Handle.GetHandle() - 1];
+	}
+	return nullptr;
+}
+
+FString FOnlineSessionIdRegistryNull::ToLogString(const FOnlineSessionIdHandle& Handle) const
+{
+	if (const FOnlineSessionIdString* Id = GetInternal(Handle))
+	{
+		return Id->Data;
+	}
+
+	return FString(TEXT("[InvalidSessionID]"));
+}
+
+TArray<uint8> FOnlineSessionIdRegistryNull::ToReplicationData(const FOnlineSessionIdHandle& Handle) const
+{
+	if (const FOnlineSessionIdString* Id = GetInternal(Handle))
+	{
+		TArray<uint8> ReplicationData;
+		ReplicationData.SetNumUninitialized(Id->Data.Len());
+		StringToBytes(Id->Data, ReplicationData.GetData(), Id->Data.Len());
+		return ReplicationData;
+	}
+
+	return TArray<uint8>();;
+}
+
+FOnlineSessionIdHandle FOnlineSessionIdRegistryNull::FromReplicationData(const TArray<uint8>& ReplicationData)
+{
+	FString Result = BytesToString(ReplicationData.GetData(), ReplicationData.Num());
+	if (Result.Len() > 0)
+	{
+		return FindOrAdd(Result);
+	}
+	return FOnlineSessionIdHandle();
+}
+
+/** FSessionNull */
+
+FSessionNull::FSessionNull()
+	: FSessionNull(FSession())
+{
+}
+
+FSessionNull::FSessionNull(const FSession& InSession)
+	: FSession(InSession)
+{
+	bool bCanBindAll;
+	OwnerInternetAddr = ISocketSubsystem::Get(PLATFORM_SOCKETSUBSYSTEM)->GetLocalHostAddr(*GLog, bCanBindAll);
+
+	// The below is a workaround for systems that set hostname to a distinct address from 127.0.0.1 on a loop back interface.
+	// See e.g. https://www.debian.org/doc/manuals/debian-reference/ch05.en.html#_the_hostname_resolution
+	// and http://serverfault.com/questions/363095/why-does-my-hostname-appear-with-the-address-127-0-1-1-rather-than-127-0-0-1-in
+	// Since we bind to 0.0.0.0, we won't answer on 127.0.1.1, so we need to advertise ourselves as 127.0.0.1 for any other loop back address we may have.
+
+	uint32 HostIp = 0;
+	OwnerInternetAddr->GetIp(HostIp); // Will return in host order
+
+	// If this address is on loop back interface, advertise it as 127.0.0.1
+	if ((HostIp & 0xff000000) == 0x7f000000)
+	{
+		OwnerInternetAddr->SetIp(0x7f000001); // 127.0.0.1
+	}
+
+#if WITH_ENGINE
+	// Now set the port that was configured
+	OwnerInternetAddr->SetPort(GetPortFromNetDriver(NAME_None));
+#endif
+
+	if (OwnerInternetAddr->GetPort() == 0)
+	{
+		OwnerInternetAddr->SetPort(7777); // Default port
+	}
+}
+
+/** FSessionsNull */
+
+FSessionsNull::FSessionsNull(FOnlineServicesNull& InServices)
+	: Super(InServices)
+	, Services(InServices)
+{
+	if (!LANSessionManager.IsValid())
+	{
+		LANSessionManager = MakeShareable(new FLANSession());
+	}
+}
+
+void FSessionsNull::Tick(float DeltaSeconds)
+{
+	if (LANSessionManager.IsValid() && LANSessionManager->GetBeaconState() > ELanBeaconState::NotUsingLanBeacon)
+	{
+		LANSessionManager->Tick(DeltaSeconds);
+	}
+
+	Super::Tick(DeltaSeconds);
+}
+
+TOnlineResult<FGetAllSessions> FSessionsNull::GetAllSessions(FGetAllSessions::Params&& Params)
+{
+	FGetAllSessions::Result Result;
+
+	for (TMap<FName, TSharedRef<FSessionNull>>::TConstIterator It(SessionsByName); It; ++It)
+	{
+		Result.Sessions.Add(It.Value());
+	}
+
+	return TOnlineResult<FGetAllSessions>(MoveTemp(Result));
+}
+
+TOnlineResult<FGetSessionByName> FSessionsNull::GetSessionByName(FGetSessionByName::Params&& Params)
+{
+	if (TSharedRef<FSessionNull>* FoundSession = SessionsByName.Find(Params.LocalName))
+	{
+		return TOnlineResult<FGetSessionByName>({ *FoundSession });
+	}
+	else
+	{
+		return TOnlineResult<FGetSessionByName>(Errors::NotFound());
+	}
+}
+
+TOnlineResult<FGetSessionById> FSessionsNull::GetSessionById(FGetSessionById::Params&& Params)
+{
+	if (TSharedRef<FSessionNull>* FoundSession = SessionsById.Find(Params.IdHandle))
+	{
+		return TOnlineResult<FGetSessionById>({ *FoundSession });
+	}
+	else
+	{
+		return TOnlineResult<FGetSessionById>(Errors::NotFound());
+	}
+}
+
+TOnlineAsyncOpHandle<FCreateSession> FSessionsNull::CreateSession(FCreateSession::Params&& Params)
+{
+	TOnlineAsyncOpRef<FCreateSession> Op = GetOp<FCreateSession>(MoveTemp(Params));
+	const FCreateSession::Params& OpParams = Op->GetParams();
+
+	if (TSharedRef<FSessionNull>* FoundSession = SessionsByName.Find(OpParams.SessionName))
+	{
+		UE_LOG(LogTemp, Warning, TEXT("[FSessionsNull::CreateSession] Could not create session with name [%s]. A session with that name already exists"), *OpParams.SessionName.ToString());
+
+		Op->SetError(Errors::InvalidParams()); // TODO: New error: Session with name %s already exists
+		return Op->GetHandle();
+	}
+
+	if (OpParams.SessionSettings.NumMaxPrivateConnections <= 0 && OpParams.SessionSettings.NumMaxPublicConnections <= 0)
+	{
+		UE_LOG(LogTemp, Warning, TEXT("[FSessionsNull::CreateSession] Could not create session with no valid NumMaxPrivateConnections [%d] or NumMaxPublicConnections [%d]"), OpParams.SessionSettings.NumMaxPrivateConnections, OpParams.SessionSettings.NumMaxPublicConnections);
+
+		Op->SetError(Errors::InvalidParams());
+		return Op->GetHandle();
+	}
+
+	// We'll only host on the LAN beacon for public sessions
+	if (OpParams.SessionSettings.JoinPolicy == ESessionJoinPolicy::Public)
+	{
+		if (!TryHostLANSession())
+		{
+			Op->SetError(Errors::RequestFailure()); // TODO: May need a new error type
+			return Op->GetHandle();
+		}
+	}
+
+	Op->Then([this](TOnlineAsyncOp<FCreateSession>& Op) mutable
+	{
+		const FCreateSession::Params& OpParams = Op.GetParams();
+
+		FSessionNull NewSession = FSessionNull();
+		NewSession.CurrentState = ESessionState::Valid;
+		NewSession.OwnerUserId = OpParams.LocalUserId;
+		NewSession.SessionId = FOnlineSessionIdRegistryNull::Get().GetNext();
+		NewSession.SessionSettings = OpParams.SessionSettings;
+
+		// For Null sessions, we'll add all the Session members manually instead of calling JoinSession since there is no API calls involved
+		NewSession.SessionSettings.SessionMembers.Append(OpParams.LocalUsers);
+	
+		TSharedRef<FSessionNull> NewSessionNullRef = MakeShared<FSessionNull>(NewSession);
+
+		SessionsByName.Add(OpParams.SessionName, NewSessionNullRef);
+		SessionsById.Add(NewSession.SessionId, NewSessionNullRef);
+
+		Op.SetResult(FCreateSession::Result{ NewSessionNullRef });
+	})
+	.Enqueue(GetSerialQueue());
+
+	return Op->GetHandle();
+}
+
+TOnlineAsyncOpHandle<FUpdateSession> FSessionsNull::UpdateSession(FUpdateSession::Params&& Params)
+{
+	TOnlineAsyncOpRef<FUpdateSession> Op = GetOp<FUpdateSession>(MoveTemp(Params));
+
+	Op->Then([this](TOnlineAsyncOp<FUpdateSession>& Op) mutable
+	{
+		const FUpdateSession::Params& OpParams = Op.GetParams();
+
+		if (TSharedRef<FSessionNull>* FoundSession = SessionsByName.Find(OpParams.SessionName))
+		{
+			// If the session exists, we'll check and update all settings one by one, with additional logic wherever required
+
+			FSessionSettings& SessionSettings = FoundSession->Get().SessionSettings;
+			const FSessionSettingsUpdate& UpdatedSettings = OpParams.Mutations;
+
+			if (UpdatedSettings.bAllowNewMembers.IsSet())
+			{
+				SessionSettings.bAllowNewMembers = UpdatedSettings.bAllowNewMembers.GetValue();
+			}
+
+			if (UpdatedSettings.bAllowSanctionedPlayers.IsSet())
+			{
+				SessionSettings.bAllowSanctionedPlayers = UpdatedSettings.bAllowSanctionedPlayers.GetValue();
+			}
+
+			if (UpdatedSettings.bAllowUnregisteredPlayers.IsSet())
+			{
+				SessionSettings.bAllowUnregisteredPlayers = UpdatedSettings.bAllowUnregisteredPlayers.GetValue();
+			}
+
+			if (UpdatedSettings.bAntiCheatProtected.IsSet())
+			{
+				SessionSettings.bAntiCheatProtected = UpdatedSettings.bAntiCheatProtected.GetValue();
+			}
+
+			if (UpdatedSettings.bPresenceEnabled.IsSet())
+			{
+				SessionSettings.bPresenceEnabled = UpdatedSettings.bPresenceEnabled.GetValue();
+			}
+
+			if (UpdatedSettings.IsDedicatedServerSession.IsSet())
+			{
+				SessionSettings.IsDedicatedServerSession = UpdatedSettings.IsDedicatedServerSession.GetValue();
+			}
+
+			if (UpdatedSettings.IsLANSession.IsSet())
+			{
+				SessionSettings.IsLANSession = UpdatedSettings.IsLANSession.GetValue();
+			}
+
+			if (UpdatedSettings.JoinPolicy.IsSet())
+			{
+				// Changes in the join policy setting will mean we start or stop the LAN Beacon broadcasting the session information
+				// There is no FriendsNull interface at the time of this implementation, hence the binary behavior (Public/Non-Public)
+				if (SessionSettings.JoinPolicy != ESessionJoinPolicy::Public && UpdatedSettings.JoinPolicy == ESessionJoinPolicy::Public)
+				{
+					TryHostLANSession();
+				}
+				else if (SessionSettings.JoinPolicy == ESessionJoinPolicy::Public && UpdatedSettings.JoinPolicy != ESessionJoinPolicy::Public)
+				{
+					StopLANSession();
+				}
+
+				SessionSettings.JoinPolicy = UpdatedSettings.JoinPolicy.GetValue();
+			}
+
+			if (UpdatedSettings.NumMaxPrivateConnections.IsSet())
+			{
+				SessionSettings.NumMaxPrivateConnections = UpdatedSettings.NumMaxPrivateConnections.GetValue();
+			}
+
+			if (UpdatedSettings.NumMaxPublicConnections.IsSet())
+			{
+				SessionSettings.NumMaxPublicConnections = UpdatedSettings.NumMaxPublicConnections.GetValue();
+			}
+
+			if (UpdatedSettings.NumOpenPrivateConnections.IsSet())
+			{
+				SessionSettings.NumOpenPrivateConnections = UpdatedSettings.NumOpenPrivateConnections.GetValue();
+			}
+
+			if (UpdatedSettings.NumOpenPublicConnections.IsSet())
+			{
+				SessionSettings.NumOpenPublicConnections = UpdatedSettings.NumOpenPublicConnections.GetValue();
+			}
+
+			// TODO: We may need some additional logic for schema changes
+			if (UpdatedSettings.SchemaName.IsSet())
+			{
+				SessionSettings.SchemaName = UpdatedSettings.SchemaName.GetValue();
+			}
+
+			if (UpdatedSettings.SessionIdOverride.IsSet())
+			{
+				SessionSettings.SessionIdOverride = UpdatedSettings.SessionIdOverride.GetValue();
+			}
+
+			for (const FName& Key : UpdatedSettings.RemovedCustomSettings)
+			{
+				SessionSettings.CustomSettings.Remove(Key);
+			}
+
+			SessionSettings.CustomSettings.Append(UpdatedSettings.UpdatedCustomSettings);
+
+			for (const FOnlineAccountIdHandle& Key : UpdatedSettings.RemovedRegisteredUsers)
+			{
+				SessionSettings.RegisteredUsers.Remove(Key);
+			}
+
+			SessionSettings.RegisteredUsers.Append(UpdatedSettings.UpdatedRegisteredUsers);
+
+			for (const FOnlineAccountIdHandle& Key : UpdatedSettings.RemovedSessionMembers)
+			{
+				SessionSettings.SessionMembers.Remove(Key);
+			}
+
+			for (FSessionMemberUpdatesMap::TConstIterator It(UpdatedSettings.UpdatedSessionMembers); It; ++It)
+			{
+				if (FSessionMember* SessionMember = SessionSettings.SessionMembers.Find(It.Key()))
+				{
+					const FSessionMemberUpdate& SessionMemberUpdate = It.Value();
+
+					SessionMember->MemberSettings.Append(SessionMemberUpdate.UpdatedMemberSettings);
+
+					for (const FName& Key : SessionMemberUpdate.RemovedMemberSettings)
+					{
+						SessionMember->MemberSettings.Remove(Key);
+					}
+				}
+			}
+
+			// We set the result and fire the event
+			Op.SetResult(FUpdateSession::Result{ *FoundSession });
+
+			FSessionUpdated SessionUpdatedEvent{ *FoundSession, UpdatedSettings };
+			SessionEvents.OnSessionUpdated.Broadcast(SessionUpdatedEvent);
+		}
+		else
+		{
+			Op.SetError(Errors::NotFound());
+		}
+	})
+	.Enqueue(GetSerialQueue());
+
+	return Op->GetHandle();
+}
+
+TOnlineAsyncOpHandle<FFindSessions> FSessionsNull::FindSessions(FFindSessions::Params&& Params)
+{
+	if (!CurrentSessionSearch.IsValid())
+	{
+		CurrentSessionSearch = MakeShared<FFindSessions::Result>();
+		CurrentSessionSearchHandle = GetOp<FFindSessions>(MoveTemp(Params));
+
+		FindLANSessions();
+
+		return CurrentSessionSearchHandle->GetHandle();
+	}
+	else
+	{
+		TOnlineAsyncOpRef<FFindSessions> Operation = GetOp<FFindSessions>(MoveTemp(Params));
+		Operation->SetError(Errors::AlreadyPending());
+		return Operation->GetHandle();
+	}
+}
+
+TOnlineAsyncOpHandle<FJoinSession> FSessionsNull::JoinSession(FJoinSession::Params&& Params)
+{
+	TOnlineAsyncOpRef<FJoinSession> Op = GetOp<FJoinSession>(MoveTemp(Params));
+	const FJoinSession::Params& OpParams = Op->GetParams();
+
+	// First, check if a session with that name already exists
+	if (TSharedRef<FSessionNull>* FoundSession = SessionsByName.Find(OpParams.SessionName))
+	{
+		Op->SetError(Errors::InvalidParams()); // TODO: New error: Session with name %s already exists
+		return Op->GetHandle();
+	}
+
+	const FSessionSettings& SessionSettings = OpParams.Session->SessionSettings;
+
+	TArray<FOnlineAccountIdHandle> LocalUserIds;
+	LocalUserIds.Reserve(OpParams.LocalUsers.Num());
+	OpParams.LocalUsers.GenerateKeyArray(LocalUserIds);
+
+	for (const FOnlineAccountIdHandle& LocalUserId : LocalUserIds)
+	{
+		if (SessionSettings.SessionMembers.Contains(LocalUserId))
+		{
+			Op->SetError(Errors::AccessDenied()); // TODO: Bespoke error: User %s already in session
+			return Op->GetHandle();
+		}
+	}
+
+	if (!SessionSettings.bAllowNewMembers)
+	{
+		Op->SetError(Errors::AccessDenied());
+		return Op->GetHandle();
+	}
+
+	if (!SessionSettings.bAllowUnregisteredPlayers)
+	{
+		for (const FOnlineAccountIdHandle& LocalUserId : LocalUserIds)
+		{
+			if (!SessionSettings.RegisteredUsers.Contains(LocalUserId))
+			{
+				Op->SetError(Errors::AccessDenied());
+				return Op->GetHandle();
+			}
+		}
+	}
+
+	if (SessionSettings.NumMaxPublicConnections > 0 && SessionSettings.NumOpenPublicConnections < LocalUserIds.Num())
+	{
+		Op->SetError(Errors::AccessDenied());
+		return Op->GetHandle();
+	}
+
+	if (SessionSettings.NumMaxPrivateConnections > 0 && SessionSettings.NumOpenPrivateConnections < LocalUserIds.Num())
+	{
+		Op->SetError(Errors::AccessDenied());
+		return Op->GetHandle();
+	}
+
+	// If no restrictions apply, we'll save our copy of the session, add the new players, and register it
+	Op->Then([this](TOnlineAsyncOp<FJoinSession>& Op) mutable
+	{
+		const FJoinSession::Params& OpParams = Op.GetParams();
+		const FSessionSettings& SessionSettings = OpParams.Session->SessionSettings;
+
+		TSharedRef<FSessionNull> NewSessionNullRef = MakeShared<FSessionNull>(FSessionNull(*OpParams.Session));
+		FSessionSettings NewSessionNullSettings = NewSessionNullRef->SessionSettings;
+
+		FSessionJoined SessionJoinedEvent;
+		SessionJoinedEvent.Session = NewSessionNullRef;
+
+		TArray<FOnlineAccountIdHandle> LocalUserIds;
+		LocalUserIds.Reserve(OpParams.LocalUsers.Num());
+		OpParams.LocalUsers.GenerateKeyArray(LocalUserIds);
+
+		for (const FOnlineAccountIdHandle& LocalUserId : LocalUserIds)
+		{
+			bool bReserveSlot = true;
+
+			if (const FRegisteredUser* RegisteredUser = SessionSettings.RegisteredUsers.Find(LocalUserId))
+			{
+				bReserveSlot = !RegisteredUser->bHasReservedSlot;
+			}
+
+			if (bReserveSlot)
+			{
+				if (SessionSettings.NumOpenPublicConnections > 0)
+				{
+					NewSessionNullSettings.NumOpenPublicConnections--;
+				}
+				else if (SessionSettings.NumOpenPrivateConnections > 0)
+				{
+					NewSessionNullSettings.NumOpenPrivateConnections--;
+				}
+			}
+
+			SessionJoinedEvent.LocalUserIds.Add(LocalUserId);
+		}
+
+		NewSessionNullSettings.SessionMembers.Append(OpParams.LocalUsers);
+
+		SessionsByName.Add(OpParams.SessionName, NewSessionNullRef);
+		SessionsById.Add(NewSessionNullRef->SessionId, NewSessionNullRef);
+
+		Op.SetResult(FJoinSession::Result{ NewSessionNullRef });
+
+		SessionEvents.OnSessionJoined.Broadcast(SessionJoinedEvent);
+	})
+	.Enqueue(GetSerialQueue());
+
+	return Op->GetHandle();
+}
+
+TOnlineAsyncOpHandle<FLeaveSession> FSessionsNull::LeaveSession(FLeaveSession::Params&& Params)
+{
+	TOnlineAsyncOpRef<FLeaveSession> Op = GetOp<FLeaveSession>(MoveTemp(Params));
+
+	Op->Then([this](TOnlineAsyncOp<FLeaveSession>& Op) mutable
+	{
+		const FLeaveSession::Params& OpParams = Op.GetParams();
+
+		if (TSharedRef<FSessionNull>* FoundSession = SessionsByName.Find(OpParams.SessionName))
+		{
+			// Only the host should stop the session at this point
+			if (FoundSession->Get().OwnerUserId == OpParams.LocalUserId && FoundSession->Get().SessionSettings.JoinPolicy == ESessionJoinPolicy::Public)
+			{
+				StopLANSession();
+			}
+
+			SessionsById.Remove((*FoundSession)->SessionId);
+			SessionsByName.Remove(OpParams.SessionName);
+
+			Op.SetResult(FLeaveSession::Result{ });
+
+			FSessionLeft SessionLeftEvent;
+			SessionLeftEvent.LocalUserIds.Add(OpParams.LocalUserId);
+			SessionLeftEvent.LocalUserIds.Append(OpParams.AdditionalLocalUsers);
+			SessionEvents.OnSessionLeft.Broadcast(SessionLeftEvent);
+		}
+		else
+		{
+			Op.SetError(Errors::NotFound());
+		}
+	})
+	.Enqueue(GetSerialQueue());
+
+	return Op->GetHandle();
+}
+
+// TODO: The logic for player registration vs adding a new session member is likely to be updated
+TOnlineAsyncOpHandle<FRegisterPlayers> FSessionsNull::RegisterPlayers(FRegisterPlayers::Params&& Params)
+{
+	TOnlineAsyncOpRef<FRegisterPlayers> Op = GetOp<FRegisterPlayers>(MoveTemp(Params));
+
+	Op->Then([this](TOnlineAsyncOp<FRegisterPlayers>& Op) mutable
+	{
+		const FRegisterPlayers::Params& OpParams = Op.GetParams();
+
+		if (TSharedRef<FSessionNull>* FoundSession = SessionsByName.Find(OpParams.SessionName))
+		{
+			FSessionSettings& SessionSettings = FoundSession->Get().SessionSettings;
+
+			for (const FOnlineAccountIdHandle& User : OpParams.TargetUsers)
+			{
+				FRegisteredUser& RegisteredUser = SessionSettings.RegisteredUsers.FindOrAdd(User);
+
+				// RegisterPlayers will not revoke a slot reserved for a registered user
+				if (!RegisteredUser.bHasReservedSlot && OpParams.bReserveSlot)
+				{
+					bool bSlotReserved = false;
+					if (SessionSettings.NumOpenPublicConnections > 0)
+					{
+						SessionSettings.NumOpenPublicConnections--;
+						bSlotReserved = true;
+					}
+					else if (SessionSettings.NumOpenPrivateConnections > 0)
+					{
+						SessionSettings.NumOpenPrivateConnections--;
+						bSlotReserved = true;
+					}
+
+					RegisteredUser.bHasReservedSlot = bSlotReserved;
+				}
+
+				if (const FSessionMember* SessionMember = SessionSettings.SessionMembers.Find(User))
+				{
+					RegisteredUser.bIsInSession = true;
+				}
+			}
+
+			Op.SetResult(FRegisterPlayers::Result{});
+		}
+		else
+		{
+			Op.SetError(Errors::NotFound());
+		}
+	})
+	.Enqueue(GetSerialQueue());
+
+	return Op->GetHandle();
+}
+
+// TODO: The logic for player registration removal vs removing a session member is likely to be updated
+TOnlineAsyncOpHandle<FUnregisterPlayers> FSessionsNull::UnregisterPlayers(FUnregisterPlayers::Params&& Params)
+{
+	TOnlineAsyncOpRef<FUnregisterPlayers> Op = GetOp<FUnregisterPlayers>(MoveTemp(Params));
+
+	Op->Then([this](TOnlineAsyncOp<FUnregisterPlayers>& Op) mutable
+	{
+		const FUnregisterPlayers::Params& OpParams = Op.GetParams();
+
+		if (TSharedRef<FSessionNull>* FoundSession = SessionsByName.Find(OpParams.SessionName))
+		{
+			FSessionSettings& SessionSettings = FoundSession->Get().SessionSettings;
+
+			for (const FOnlineAccountIdHandle& User : OpParams.TargetUsers)
+			{
+				if (const FRegisteredUser* RegisteredUser = SessionSettings.RegisteredUsers.Find(User))
+				{
+					SessionSettings.RegisteredUsers.Remove(User);
+
+					bool bRemoveUserFromSession = RegisteredUser->bIsInSession && OpParams.bRemoveUnregisteredPlayers;
+
+					if (bRemoveUserFromSession)
+					{
+						SessionSettings.SessionMembers.Remove(User);
+					}
+
+					if (RegisteredUser->bHasReservedSlot || bRemoveUserFromSession)
+					{
+						if (SessionSettings.NumOpenPublicConnections < SessionSettings.NumMaxPublicConnections)
+						{
+							SessionSettings.NumOpenPublicConnections++;
+						}
+						else if (SessionSettings.NumOpenPrivateConnections < SessionSettings.NumMaxPrivateConnections)
+						{
+							SessionSettings.NumOpenPrivateConnections++;
+						}
+					}
+				}
+			}
+
+			Op.SetResult(FUnregisterPlayers::Result{});
+		}
+		else
+		{
+			Op.SetError(Errors::NotFound());
+		}
+	})
+	.Enqueue(GetSerialQueue());
+
+	return Op->GetHandle();
+}
+
+/** LAN Methods */
+
+bool FSessionsNull::TryHostLANSession()
+{
+	// The LAN Beacon can only broadcast one session at a time
+	if (LANSessionManager->GetBeaconState() == ELanBeaconState::NotUsingLanBeacon)
+	{
+		FOnValidQueryPacketDelegate QueryPacketDelegate = FOnValidQueryPacketDelegate::CreateRaw(this, &FSessionsNull::OnValidQueryPacketReceived);
+
+		if (LANSessionManager->Host(QueryPacketDelegate))
+		{
+			UE_LOG(LogTemp, VeryVerbose, TEXT("[FSessionsNull::TryHostLANSession] LAN Beacon hosting..."));
+
+			PublicSessionsHosted++;
+
+			return true;
+		}
+		else
+		{
+			UE_LOG(LogTemp, VeryVerbose, TEXT("[FSessionsNull::TryHostLANSession] LAN Beacon failed to host!"));
+
+			LANSessionManager->StopLANSession();
+		}
+	}
+	else
+	{
+		UE_LOG(LogTemp, VeryVerbose, TEXT("[FSessionsNull::TryHostLANSession] LAN Beacon already in use!"));
+	}
+
+	return false;
+}
+
+void FSessionsNull::FindLANSessions()
+{
+	// Recreate the unique identifier for this client
+	GenerateNonce((uint8*)&LANSessionManager->LanNonce, 8);
+
+	// Bind delegates
+	FOnValidResponsePacketDelegate ResponseDelegate = FOnValidResponsePacketDelegate::CreateRaw(this, &FSessionsNull::OnValidResponsePacketReceived);
+	FOnSearchingTimeoutDelegate TimeoutDelegate = FOnSearchingTimeoutDelegate::CreateRaw(this, &FSessionsNull::OnLANSearchTimeout);
+
+	FNboSerializeToBufferNullSvc Packet(LAN_BEACON_MAX_PACKET_SIZE);
+	LANSessionManager->CreateClientQueryPacket(Packet, LANSessionManager->LanNonce);
+	if (LANSessionManager->Search(Packet, ResponseDelegate, TimeoutDelegate) == false)
+	{
+		UE_LOG(LogTemp, Warning, TEXT("[FSessionsNull::FindLANSessions] Search failed!"));
+
+		if (LANSessionManager->GetBeaconState() == ELanBeaconState::Searching)
+		{
+			LANSessionManager->StopLANSession();
+		}
+
+		// Trigger the delegate as having failed
+		CurrentSessionSearchHandle->SetError(Errors::RequestFailure()); // TODO: May need a new error type
+
+		CurrentSessionSearch.Reset();
+		CurrentSessionSearchHandle.Reset();
+
+		// If we were hosting public sessions before the search, we'll return the beacon to that state
+		if (PublicSessionsHosted > 0)
+		{
+			TryHostLANSession();
+		}
+	}
+	else
+	{
+		UE_LOG(LogTemp, Warning, TEXT("[]FSessionsNull::FindLANSessions] Searching...."));
+	}
+}
+
+void FSessionsNull::StopLANSession()
+{
+	PublicSessionsHosted--;
+
+	if (PublicSessionsHosted <= 0)
+	{
+		LANSessionManager->StopLANSession();
+	}
+}
+
+void FSessionsNull::OnValidQueryPacketReceived(uint8* PacketData, int32 PacketLength, uint64 ClientNonce)
+{
+	// Iterate through all registered sessions and respond for each one that can be joinable
+	for (TMap<FName, TSharedRef<FSessionNull>>::TConstIterator It(SessionsByName); It; ++It)
+	{
+		const TSharedRef<FSessionNull>& Session = It.Value();
+
+		if (Session->SessionSettings.JoinPolicy == ESessionJoinPolicy::Public)
+		{
+			FNboSerializeToBufferNullSvc Packet(LAN_BEACON_MAX_PACKET_SIZE);
+			// Create the basic header before appending additional information
+			LANSessionManager->CreateHostResponsePacket(Packet, ClientNonce);
+
+			// Add all the session details
+			AppendSessionToPacket(Packet, Session);
+
+			// Broadcast this response so the client can see us
+			if (!Packet.HasOverflow())
+			{
+				LANSessionManager->BroadcastPacket(Packet, Packet.GetByteCount());
+			}
+			else
+			{
+				UE_LOG(LogTemp, Warning, TEXT("[FSessionsNull::OnValidQueryPacketReceived] LAN broadcast packet overflow, cannot broadcast on LAN"));
+			}
+		}
+	}
+}
+
+void FSessionsNull::OnValidResponsePacketReceived(uint8* PacketData, int32 PacketLength)
+{
+	if (CurrentSessionSearch.IsValid())
+	{
+		TSharedRef<FSessionNull> Session = MakeShared<FSessionNull>();
+
+		FNboSerializeFromBufferNullSvc Packet(PacketData, PacketLength);
+		ReadSessionFromPacket(Packet, Session);
+
+		CurrentSessionSearch->FoundSessions.Add(Session);
+	}
+}
+
+void FSessionsNull::OnLANSearchTimeout()
+{
+	if (LANSessionManager->GetBeaconState() == ELanBeaconState::Searching)
+	{
+		LANSessionManager->StopLANSession();
+	}
+
+	if (!CurrentSessionSearch.IsValid())
+	{
+		UE_LOG(LogTemp, Warning, TEXT("[FSessionsNull::OnLANSearchTimeout] Session search invalid"));
+
+		CurrentSessionSearchHandle->SetError(Errors::InvalidState());
+	}
+
+	if (CurrentSessionSearch->FoundSessions.IsEmpty())
+	{
+		UE_LOG(LogTemp, Warning, TEXT("[FSessionsNull::OnLANSearchTimeout] Session search found no results"));
+
+		CurrentSessionSearchHandle->SetError(Errors::NotFound());
+	}
+
+	UE_LOG(LogTemp, Warning, TEXT("[FSessionsNull::OnLANSearchTimeout] %d sessions found!"), CurrentSessionSearch->FoundSessions.Num());
+
+	FFindSessions::Result Result = *CurrentSessionSearch;
+	CurrentSessionSearchHandle->SetResult(MoveTemp(Result));
+
+	CurrentSessionSearch.Reset();
+	CurrentSessionSearchHandle.Reset();
+
+	// If we were hosting public sessions before the search, we'll return the beacon to that state
+	if (PublicSessionsHosted > 0)
+	{
+		TryHostLANSession();
+	}
+}
+
+void FSessionsNull::AppendSessionToPacket(FNboSerializeToBufferNullSvc& Packet, const TSharedRef<FSessionNull>& Session)
+{
+	// We won't save CurrentState in the packet as all advertised sessions will be Valid
+	Packet << Session->OwnerUserId
+		<< Session->SessionId
+		<< *Session->OwnerInternetAddr;
+
+	// TODO: Write session settings to packet, after SchemaVariant work
+}
+
+void FSessionsNull::ReadSessionFromPacket(FNboSerializeFromBufferNullSvc& Packet, const TSharedRef<FSessionNull>& Session)
+{
+	Packet >> Session->OwnerUserId
+		>> Session->SessionId
+		>> *Session->OwnerInternetAddr;
+
+	// We'll set the connect address for the remote session as a custom parameter, so it can be read in OnlineServices' GetResolvedConnectString
+	FCustomSessionSetting ConnectString;
+	ConnectString.Data.Set<FString>(Session->OwnerInternetAddr->ToString(true));
+	ConnectString.Visibility = ECustomSessionSettingVisibility::ViaOnlineService;
+	Session->SessionSettings.CustomSettings.Add(CONNECT_STRING_TAG, ConnectString);
+
+	// TODO:: Read session settings from packet, after SchemaVariant work
+}
+
+/* UE::Online */ }
\ No newline at end of file
diff --git a/Engine/Plugins/Online/Experimental/OnlineServicesNull/Source/Public/Online/NboSerializerNullSvc.h b/Engine/Plugins/Online/Experimental/OnlineServicesNull/Source/Public/Online/NboSerializerNullSvc.h
--- a/Engine/Plugins/Online/Experimental/OnlineServicesNull/Source/Public/Online/NboSerializerNullSvc.h
+++ b/Engine/Plugins/Online/Experimental/OnlineServicesNull/Source/Public/Online/NboSerializerNullSvc.h
@@ -5,6 +5,8 @@
 #include "CoreMinimal.h"
 #include "NboSerializer.h"
 #include "Online/AuthNull.h"
+#include "Online/LobbiesNull.h"
+#include "Online/SessionsNull.h"
 #include "Online/CoreOnline.h"
 
 /**
@@ -38,6 +40,14 @@ public:
 		return Ar;
 	}
 
+	friend inline FNboSerializeToBufferNullSvc& operator<<(FNboSerializeToBufferNullSvc& Ar, const FOnlineSessionIdHandle& SessionId)
+	{
+		TArray<uint8> Data = FOnlineSessionIdRegistryNull::Get().ToReplicationData(SessionId);
+		Ar << Data.Num();
+		Ar.WriteBinary(Data.GetData(), Data.Num());
+		return Ar;
+	}
+
 	friend inline FNboSerializeToBufferNullSvc& operator<<(FNboSerializeToBufferNullSvc& Ar, const TMap<FLobbyAttributeId, FLobbyVariant>& Map)
 	{
 		Ar << Map.Num();
@@ -75,6 +85,16 @@ public:
 		UniqueId = FOnlineAccountIdRegistryNull::Get().FromReplicationData(Data);
 		return Ar;
 	}
+
+	friend inline FNboSerializeFromBufferNullSvc& operator>>(FNboSerializeFromBufferNullSvc& Ar, FOnlineSessionIdHandle& SessionId)
+	{
+		TArray<uint8> Data;
+		int32 Size;
+		Ar >> Size;
+		Ar.ReadBinaryArray(Data, Size);
+		SessionId = FOnlineSessionIdRegistryNull::Get().FromReplicationData(Data);
+		return Ar;
+	}
 	
 	friend inline FNboSerializeFromBufferNullSvc& operator>>(FNboSerializeFromBufferNullSvc& Ar, TMap<FLobbyAttributeId, FLobbyVariant>& Map)
 	{
diff --git a/Engine/Plugins/Online/Experimental/OnlineServicesNull/Source/Public/Online/OnlineServicesNull.h b/Engine/Plugins/Online/Experimental/OnlineServicesNull/Source/Public/Online/OnlineServicesNull.h
--- a/Engine/Plugins/Online/Experimental/OnlineServicesNull/Source/Public/Online/OnlineServicesNull.h
+++ b/Engine/Plugins/Online/Experimental/OnlineServicesNull/Source/Public/Online/OnlineServicesNull.h
@@ -6,6 +6,8 @@
 
 namespace UE::Online {
 
+static FName CONNECT_STRING_TAG = TEXT("CONNECT_STRING");
+
 class ONLINESERVICESNULL_API FOnlineServicesNull : public FOnlineServicesCommon
 {
 public:
diff --git a/Engine/Plugins/Online/Experimental/OnlineServicesNull/Source/Public/Online/SessionsNull.h b/Engine/Plugins/Online/Experimental/OnlineServicesNull/Source/Public/Online/SessionsNull.h
new file mode 100644
--- /dev/null
+++ b/Engine/Plugins/Online/Experimental/OnlineServicesNull/Source/Public/Online/SessionsNull.h
@@ -0,0 +1,101 @@
+// Copyright Epic Games, Inc. All Rights Reserved.
+#pragma once
+
+#include "CoreMinimal.h"
+#include "LANBeacon.h"
+#include "Online/SessionsCommon.h"
+
+namespace UE::Online {
+
+class FOnlineServicesNull;
+
+class FOnlineSessionIdString
+{
+public:
+	FString Data;
+	FOnlineSessionIdHandle Handle;
+};
+
+class FOnlineSessionIdRegistryNull : public IOnlineSessionIdRegistry
+{
+public:
+	static FOnlineSessionIdRegistryNull& Get();
+
+	FOnlineSessionIdHandle Find(FString SessionId) const;
+	FOnlineSessionIdHandle FindOrAdd(FString SessionId);
+	FOnlineSessionIdHandle GetNext();
+
+	// Begin IOnlineSessionIdRegistry
+	virtual FString ToLogString(const FOnlineSessionIdHandle& Handle) const override;
+	virtual TArray<uint8> ToReplicationData(const FOnlineSessionIdHandle& Handle) const override;
+	virtual FOnlineSessionIdHandle FromReplicationData(const TArray<uint8>& ReplicationString) override;
+	// End IOnlineSessionIdRegistry
+
+	virtual ~FOnlineSessionIdRegistryNull() = default;
+
+private:
+	const FOnlineSessionIdString* GetInternal(const FOnlineSessionIdHandle& Handle) const;
+
+	TArray<FOnlineSessionIdString> Ids;
+	TMap<FString, FOnlineSessionIdString> StringToId;
+};
+
+class FSessionNull : public FSession
+{
+public:
+	FSessionNull();
+	FSessionNull(const FSession& InSession);
+
+public:
+	/** The IP address for the session owner */
+	TSharedPtr<FInternetAddr> OwnerInternetAddr;
+};
+
+class FSessionsNull : public FSessionsCommon
+{
+public:
+	using Super = FSessionsCommon;
+
+	FSessionsNull(FOnlineServicesNull& InServices);
+	virtual ~FSessionsNull() = default;
+
+	virtual void Tick(float DeltaSeconds) override;
+
+	// ISessions
+	virtual TOnlineResult<FGetAllSessions> GetAllSessions(FGetAllSessions::Params&& Params) override;
+	virtual TOnlineResult<FGetSessionByName> GetSessionByName(FGetSessionByName::Params&& Params) override;
+	virtual TOnlineResult<FGetSessionById> GetSessionById(FGetSessionById::Params&& Params) override;
+	virtual TOnlineAsyncOpHandle<FCreateSession> CreateSession(FCreateSession::Params&& Params) override;
+	virtual TOnlineAsyncOpHandle<FUpdateSession> UpdateSession(FUpdateSession::Params&& Params) override;
+	virtual TOnlineAsyncOpHandle<FFindSessions> FindSessions(FFindSessions::Params&& Params) override;
+	virtual TOnlineAsyncOpHandle<FJoinSession> JoinSession(FJoinSession::Params&& Params) override;
+	virtual TOnlineAsyncOpHandle<FLeaveSession> LeaveSession(FLeaveSession::Params&& Params) override;
+	virtual TOnlineAsyncOpHandle<FRegisterPlayers> RegisterPlayers(FRegisterPlayers::Params&& Params) override;
+	virtual TOnlineAsyncOpHandle<FUnregisterPlayers> UnregisterPlayers(FUnregisterPlayers::Params&& Params) override;
+
+private:
+	/** LAN Methods */
+	bool TryHostLANSession();
+	void FindLANSessions();
+	void StopLANSession();
+	void OnValidQueryPacketReceived(uint8* PacketData, int32 PacketLength, uint64 ClientNonce);
+	void OnValidResponsePacketReceived(uint8* PacketData, int32 PacketLength);
+	void OnLANSearchTimeout();
+	void AppendSessionToPacket(class FNboSerializeToBufferNullSvc& Packet, const TSharedRef<FSessionNull>& Session);
+	void ReadSessionFromPacket(class FNboSerializeFromBufferNullSvc& Packet, const TSharedRef<FSessionNull>& Session);
+	
+private:
+	FOnlineServicesNull& Services;
+
+	TMap<FName, TSharedRef<FSessionNull>> SessionsByName;
+	TMap<FOnlineSessionIdHandle, TSharedRef<FSessionNull>> SessionsById;
+
+	TSharedPtr<FFindSessions::Result> CurrentSessionSearch;
+	TSharedPtr<TOnlineAsyncOp<FFindSessions>> CurrentSessionSearchHandle;
+
+	uint32 PublicSessionsHosted = 0;
+
+	TSharedPtr<FLANSession> LANSessionManager;
+};
+
+/* UE::Online */ }
\ No newline at end of file
