OSSv2: Add happy path tests for stats integration tests;
Fix StatsNull to make sure it passes the tests.

diff --git a/Engine/Plugins/Online/Experimental/OnlineServicesNull/Source/Private/Online/StatsNull.cpp b/Engine/Plugins/Online/Experimental/OnlineServicesNull/Source/Private/Online/StatsNull.cpp
--- a/Engine/Plugins/Online/Experimental/OnlineServicesNull/Source/Private/Online/StatsNull.cpp
+++ b/Engine/Plugins/Online/Experimental/OnlineServicesNull/Source/Private/Online/StatsNull.cpp
@@ -77,7 +77,7 @@ TOnlineAsyncOpHandle<FUpdateStats> FStatsNull::UpdateStats(FUpdateStats::Params&
 						}
 						else
 						{
-							UserStats[UpdateUserStatPair.Key] = UpdateUserStatPair.Value;
+							UserStats.Emplace(UpdateUserStatPair.Key, UpdateUserStatPair.Value);
 						}
 					}
 
@@ -116,7 +116,9 @@ TOnlineAsyncOpHandle<FQueryStats> FStatsNull::QueryStats(FQueryStats::Params&& P
 		}
 
 		InAsyncOp.SetResult(MoveTemp(Result));
-	});
+	})
+	.Enqueue(GetSerialQueue());
+
 	return Op->GetHandle();
 }
 
@@ -142,7 +144,8 @@ TOnlineAsyncOpHandle<FBatchQueryStats> FStatsNull::BatchQueryStats(FBatchQuerySt
 		}
 
 		InAsyncOp.SetResult(MoveTemp(Result));
-	});
+	})
+	.Enqueue(GetSerialQueue());
 
 	return Op->GetHandle();
 }
@@ -166,7 +169,8 @@ TOnlineAsyncOpHandle<FResetStats> FStatsNull::ResetStats(FResetStats::Params&& P
 			UsersStats.RemoveAt(Index);
 		}
 		InAsyncOp.SetResult({});
-	});
+	})
+	.Enqueue(GetSerialQueue());
 
 	return Op->GetHandle();
 }
