OSSv2 Sessions Interface Refactor - Moving common members, methods and parameter and state checks to Common implementation

diff --git a/Engine/Plugins/Online/Experimental/OnlineServices/Source/OnlineServicesCommon/Private/Online/SessionsCommon.cpp b/Engine/Plugins/Online/Experimental/OnlineServices/Source/OnlineServicesCommon/Private/Online/SessionsCommon.cpp
--- a/Engine/Plugins/Online/Experimental/OnlineServices/Source/OnlineServicesCommon/Private/Online/SessionsCommon.cpp
+++ b/Engine/Plugins/Online/Experimental/OnlineServices/Source/OnlineServicesCommon/Private/Online/SessionsCommon.cpp
@@ -7,6 +7,7 @@
 
 namespace UE::Online {
 
+	/** FSessionsCommon */
 	FSessionsCommon::FSessionsCommon(FOnlineServicesCommon& InServices)
 		: TOnlineComponent(TEXT("Sessions"), InServices)
 	{
@@ -37,19 +38,40 @@
 		RegisterCommand(&FSessionsCommon::UnregisterPlayers);
 	}
 
-	TOnlineResult<FGetAllSessions> FSessionsCommon::GetAllSessions(FGetAllSessions::Params&& Params)
+	TOnlineResult<FGetAllSessions> FSessionsCommon::GetAllSessions(FGetAllSessions::Params&& Params) const
 	{
-		return TOnlineResult<FGetAllSessions>(Errors::NotImplemented());
+		FGetAllSessions::Result Result;
+
+		for (const TPair<FName, TSharedRef<FSession>>& SessionPair : SessionsByName)
+		{
+			Result.Sessions.Add(SessionPair.Value);
+		}
+
+		return TOnlineResult<FGetAllSessions>(MoveTemp(Result));
 	}
 
-	TOnlineResult<FGetSessionByName> FSessionsCommon::GetSessionByName(FGetSessionByName::Params&& Params)
+	TOnlineResult<FGetSessionByName> FSessionsCommon::GetSessionByName(FGetSessionByName::Params&& Params) const
 	{
-		return TOnlineResult<FGetSessionByName>(Errors::NotImplemented());
+		if (const TSharedRef<FSession>* FoundSession = SessionsByName.Find(Params.LocalName))
+		{
+			return TOnlineResult<FGetSessionByName>({ *FoundSession });
+		}
+		else
+		{
+			return TOnlineResult<FGetSessionByName>(Errors::NotFound());
+		}
 	}
 
-	TOnlineResult<FGetSessionById> FSessionsCommon::GetSessionById(FGetSessionById::Params&& Params)
+	TOnlineResult<FGetSessionById> FSessionsCommon::GetSessionById(FGetSessionById::Params&& Params) const
 	{
-		return TOnlineResult<FGetSessionById>(Errors::NotImplemented());
+		if (const TSharedRef<FSession>* FoundSession = SessionsById.Find(Params.IdHandle))
+		{
+			return TOnlineResult<FGetSessionById>({ *FoundSession });
+		}
+		else
+		{
+			return TOnlineResult<FGetSessionById>(Errors::NotFound());
+		}
 	}
 
 	TOnlineAsyncOpHandle<FCreateSession> FSessionsCommon::CreateSession(FCreateSession::Params&& Params)
@@ -115,18 +137,121 @@
 		return Operation->GetHandle();
 	}
 
+	// TODO: The logic for player registration vs adding a new session member is likely to be updated
 	TOnlineAsyncOpHandle<FRegisterPlayers> FSessionsCommon::RegisterPlayers(FRegisterPlayers::Params&& Params)
 	{
-		TOnlineAsyncOpRef<FRegisterPlayers> Operation = GetOp<FRegisterPlayers>(MoveTemp(Params));
-		Operation->SetError(Errors::NotImplemented());
-		return Operation->GetHandle();
+		TOnlineAsyncOpRef<FRegisterPlayers> Op = GetOp<FRegisterPlayers>(MoveTemp(Params));
+
+		Op->Then([this](TOnlineAsyncOp<FRegisterPlayers>& Op) mutable
+		{
+			const FRegisterPlayers::Params& OpParams = Op.GetParams();
+
+			if (TSharedRef<FSession>* FoundSession = SessionsByName.Find(OpParams.SessionName))
+			{
+				FSessionSettings& SessionSettings = FoundSession->Get().SessionSettings;
+
+				for (const FOnlineAccountIdHandle& User : OpParams.TargetUsers)
+				{
+					FRegisteredUser& RegisteredUser = SessionSettings.RegisteredUsers.FindOrAdd(User);
+
+					// RegisterPlayers will not revoke a slot reserved for a registered user
+					if (!RegisteredUser.bHasReservedSlot && OpParams.bReserveSlot)
+					{
+						bool bSlotReserved = false;
+						if (SessionSettings.NumOpenPublicConnections > 0)
+						{
+							SessionSettings.NumOpenPublicConnections--;
+							bSlotReserved = true;
+						}
+						else if (SessionSettings.NumOpenPrivateConnections > 0)
+						{
+							SessionSettings.NumOpenPrivateConnections--;
+							bSlotReserved = true;
+						}
+
+						RegisteredUser.bHasReservedSlot = bSlotReserved;
+					}
+
+					if (const FSessionMember* SessionMember = SessionSettings.SessionMembers.Find(User))
+					{
+						RegisteredUser.bIsInSession = true;
+					}
+				}
+
+				Op.SetResult(FRegisterPlayers::Result{});
+			}
+			else
+			{
+				UE_LOG(LogTemp, Warning, TEXT("[FSessionsCommon::RegisterPlayers] Unable to register players. Session with name [%s] not found"), *OpParams.SessionName.ToString());
+
+				Op.SetError(Errors::NotFound());
+			}
+		})
+		.Enqueue(GetSerialQueue());
+
+		return Op->GetHandle();
 	}
 
+	// TODO: The logic for player registration removal vs removing a session member is likely to be updated
 	TOnlineAsyncOpHandle<FUnregisterPlayers> FSessionsCommon::UnregisterPlayers(FUnregisterPlayers::Params&& Params)
 	{
-		TOnlineAsyncOpRef<FUnregisterPlayers> Operation = GetOp<FUnregisterPlayers>(MoveTemp(Params));
-		Operation->SetError(Errors::NotImplemented());
-		return Operation->GetHandle();
+		TOnlineAsyncOpRef<FUnregisterPlayers> Op = GetOp<FUnregisterPlayers>(MoveTemp(Params));
+
+		Op->Then([this](TOnlineAsyncOp<FUnregisterPlayers>& Op) mutable
+		{
+			const FUnregisterPlayers::Params& OpParams = Op.GetParams();
+
+			if (TSharedRef<FSession>* FoundSession = SessionsByName.Find(OpParams.SessionName))
+			{
+				FSessionSettings& SessionSettings = FoundSession->Get().SessionSettings;
+
+				for (const FOnlineAccountIdHandle& User : OpParams.TargetUsers)
+				{
+					if (const FRegisteredUser* RegisteredUser = SessionSettings.RegisteredUsers.Find(User))
+					{
+						SessionSettings.RegisteredUsers.Remove(User);
+
+						// TODO: This logic will move to RemoveSessionMember
+						bool bRemoveUserFromSession = RegisteredUser->bIsInSession && OpParams.bRemoveUnregisteredPlayers;
+
+						if (bRemoveUserFromSession)
+						{
+							SessionSettings.SessionMembers.Remove(User);
+						}
+
+						if (RegisteredUser->bHasReservedSlot || bRemoveUserFromSession)
+						{
+							if (SessionSettings.NumOpenPublicConnections < SessionSettings.NumMaxPublicConnections)
+							{
+								SessionSettings.NumOpenPublicConnections++;
+							}
+							else if (SessionSettings.NumOpenPrivateConnections < SessionSettings.NumMaxPrivateConnections)
+							{
+								SessionSettings.NumOpenPrivateConnections++;
+							}
+						}
+					}
+					else
+					{
+						UE_LOG(LogTemp, Warning, TEXT("[FSessionsCommon::UnregisterPlayers] Unable to unregister player with id [%s]. Player not registered"), *ToLogString(User));
+
+						Op.SetError(Errors::NotFound());
+						return;
+					}
+				}
+
+				Op.SetResult(FUnregisterPlayers::Result{});
+			}
+			else
+			{
+				UE_LOG(LogTemp, Warning, TEXT("[FSessionsCommon::UnregisterPlayers] Unable to unregister players. Session with name [%s] not found"), *OpParams.SessionName.ToString());
+
+				Op.SetError(Errors::NotFound());
+			}
+		})
+		.Enqueue(GetSerialQueue());
+
+		return Op->GetHandle();
 	}
 
 	TOnlineEvent<void(const FSessionJoined&)> FSessionsCommon::OnSessionJoined()
@@ -154,4 +279,285 @@
 		return SessionEvents.OnInviteAccepted;
 	}
 
+	FOnlineError FSessionsCommon::CheckCreateSessionParams(const FCreateSession::Params& Params)
+	{
+		if (Params.SessionSettings.NumMaxPrivateConnections == 0 && Params.SessionSettings.NumMaxPublicConnections == 0)
+		{
+			UE_LOG(LogTemp, Warning, TEXT("[FSessionsCommon::CheckCreateSessionParams] Could not create session with no valid NumMaxPrivateConnections [%d] or NumMaxPublicConnections [%d]"), Params.SessionSettings.NumMaxPrivateConnections, Params.SessionSettings.NumMaxPublicConnections);
+
+			return Errors::InvalidParams();
+		}
+
+		return Errors::Success();
+	}
+
+	FOnlineError FSessionsCommon::CheckCreateSessionState(const FCreateSession::Params& Params)
+	{
+		if (TSharedRef<FSession>* FoundSession = SessionsByName.Find(Params.SessionName))
+		{
+			UE_LOG(LogTemp, Warning, TEXT("[FSessionsCommon::CheckCreateSessionState] Could not create session with name [%s]. A session with that name already exists"), *Params.SessionName.ToString());
+
+			return Errors::InvalidState();
+		}
+
+		if (!Params.SessionSettings.IsDedicatedServerSession)
+		{
+			IAuthPtr Auth = Services.GetAuthInterface();
+
+			FAuthGetAccountByAccountId::Params GetAccountByAccountIParams{ Params.LocalUserId };
+			TOnlineResult<FAuthGetAccountByAccountId> Result = Auth->GetAccountByAccountId(MoveTemp(GetAccountByAccountIParams));
+
+			if (Result.IsOk())
+			{
+				if (Result.GetOkValue().AccountInfo->LoginStatus == ELoginStatus::NotLoggedIn)
+				{
+					UE_LOG(LogTemp, Warning, TEXT("[FSessionsCommon::CheckCreateSessionState] Could not create session with user [%s] not logged in"), *ToLogString(Params.LocalUserId));
+
+					return Errors::InvalidUser();
+				}
+			}
+		}
+
+		// User login check for all local users
+		IAuthPtr Auth = Services.GetAuthInterface();
+
+		TArray<FOnlineAccountIdHandle> LocalUserIds;
+		LocalUserIds.Reserve(Params.LocalUsers.Num());
+		Params.LocalUsers.GenerateKeyArray(LocalUserIds);
+
+		for (const FOnlineAccountIdHandle& LocalUserId : LocalUserIds)
+		{
+			FAuthGetAccountByAccountId::Params GetAccountByAccountIParams{ LocalUserId };
+			TOnlineResult<FAuthGetAccountByAccountId> Result = Auth->GetAccountByAccountId(MoveTemp(GetAccountByAccountIParams));
+
+			if (Result.IsOk())
+			{
+				if (Result.GetOkValue().AccountInfo->LoginStatus == ELoginStatus::NotLoggedIn)
+				{
+					UE_LOG(LogTemp, Warning, TEXT("[FSessionsCommon::CheckCreateSessionState] Could not create session with user [%s] not logged in"), *ToLogString(LocalUserId));
+
+					return Errors::InvalidUser();
+				}
+			}
+		}
+
+		return Errors::Success();
+	}
+
+	FOnlineError FSessionsCommon::CheckUpdateSessionState(const FUpdateSession::Params& Params)
+	{
+		if (TSharedRef<FSession>* FoundSession = SessionsByName.Find(Params.SessionName))
+		{
+			if (!FoundSession->Get().SessionSettings.IsDedicatedServerSession)
+			{
+				IAuthPtr Auth = Services.GetAuthInterface();
+
+				FAuthGetAccountByAccountId::Params GetAccountByAccountIParams{ Params.LocalUserId };
+				TOnlineResult<FAuthGetAccountByAccountId> Result = Auth->GetAccountByAccountId(MoveTemp(GetAccountByAccountIParams));
+
+				if (Result.IsOk())
+				{
+					if (Result.GetOkValue().AccountInfo->LoginStatus == ELoginStatus::NotLoggedIn)
+					{
+						UE_LOG(LogTemp, Warning, TEXT("[FSessionsNull::CheckUpdateSessionState] Could not update session with user [%s] not logged in"), *ToLogString(Params.LocalUserId));
+
+						return Errors::InvalidUser();
+					}
+				}
+			}
+		}
+		else
+		{
+			UE_LOG(LogTemp, Warning, TEXT("[FSessionsCommon::CheckUpdateSessionState] Could not update session with name [%s]. Session not found"), *Params.SessionName.ToString());
+
+			return Errors::NotFound();
+		}
+
+		return Errors::Success();
+	}
+
+	FOnlineError FSessionsCommon::CheckFindSessionsParams(const FFindSessions::Params& Params)
+	{
+		if (Params.MaxResults <= 0)
+		{
+			UE_LOG(LogTemp, Warning, TEXT("[FSessionsCommon::CheckFindSessionsParams] Could not find sessions with no valid MaxResults [%d]"), Params.MaxResults);
+
+			return Errors::InvalidParams();
+		}
+
+		return Errors::Success();
+	}
+
+	FOnlineError FSessionsCommon::CheckFindSessionsState(const FFindSessions::Params& Params)
+	{
+		// User login check
+		IAuthPtr Auth = Services.GetAuthInterface();
+
+		FAuthGetAccountByAccountId::Params GetAccountByAccountIParams{ Params.LocalUserId };
+		TOnlineResult<FAuthGetAccountByAccountId> Result = Auth->GetAccountByAccountId(MoveTemp(GetAccountByAccountIParams));
+
+		if (Result.IsOk())
+		{
+			if (Result.GetOkValue().AccountInfo->LoginStatus == ELoginStatus::NotLoggedIn)
+			{
+				UE_LOG(LogTemp, Warning, TEXT("[FSessionsCommon::CheckFindSessionsState] Could not find sessions with user [%s] not logged in"), *ToLogString(Params.LocalUserId));
+
+				return Errors::InvalidUser();
+			}
+		}
+
+		// Ongoing search check
+		if (CurrentSessionSearch.IsValid())
+		{
+			UE_LOG(LogTemp, Warning, TEXT("[FSessionsCommon::CheckFindSessionsState] Could not find sessions, search already in progress"));
+
+			return Errors::AlreadyPending();
+		}
+
+		return Errors::Success();
+	}
+
+	FOnlineError FSessionsCommon::CheckJoinSessionParams(const FJoinSession::Params& Params)
+	{
+		const FSessionSettings& SessionSettings = Params.Session->SessionSettings;
+
+		TArray<FOnlineAccountIdHandle> LocalUserIds;
+		LocalUserIds.Reserve(Params.LocalUsers.Num());
+		Params.LocalUsers.GenerateKeyArray(LocalUserIds);
+
+		for (const FOnlineAccountIdHandle& LocalUserId : LocalUserIds)
+		{
+			if (SessionSettings.SessionMembers.Contains(LocalUserId))
+			{
+				UE_LOG(LogTemp, Warning, TEXT("[FSessionsCommon::CheckJoinSessionParams] Could not join session. User [%s] already in session"), *ToLogString(LocalUserId));
+
+				return Errors::AccessDenied();
+			}
+		}
+
+		if (!SessionSettings.bAllowNewMembers)
+		{
+			UE_LOG(LogTemp, Warning, TEXT("[FSessionsCommon::CheckJoinSessionParams] Could not join session. bAllowNewMembers is set to false"));
+
+			return Errors::AccessDenied();
+		}
+
+		if (!SessionSettings.bAllowUnregisteredPlayers)
+		{
+			for (const FOnlineAccountIdHandle& LocalUserId : LocalUserIds)
+			{
+				if (!SessionSettings.RegisteredUsers.Contains(LocalUserId))
+				{
+					UE_LOG(LogTemp, Warning, TEXT("[FSessionsCommon::CheckJoinSessionParams] Could not join session. bAllowUnregisteredPlayers is set to false and user [%s] is not registered"), *ToLogString(LocalUserId));
+
+					return Errors::AccessDenied();
+				}
+			}
+		}
+
+		if (SessionSettings.NumMaxPublicConnections > 0 && SessionSettings.NumOpenPublicConnections < (uint32)LocalUserIds.Num())
+		{
+			UE_LOG(LogTemp, Warning, TEXT("[FSessionsCommon::CheckJoinSessionParams] Could not join session. Not enough NumOpenPublicConnections [%s] for the amount of joining players [%d]"), SessionSettings.NumOpenPublicConnections, LocalUserIds.Num());
+
+			return Errors::AccessDenied();
+		}
+
+		if (SessionSettings.NumMaxPrivateConnections > 0 && SessionSettings.NumOpenPrivateConnections < (uint32)LocalUserIds.Num())
+		{
+			UE_LOG(LogTemp, Warning, TEXT("[FSessionsCommon::CheckJoinSessionParams] Could not join session. Not enough NumOpenPrivateConnections [%s] for the amount of joining players [%d]"), SessionSettings.NumOpenPrivateConnections, LocalUserIds.Num());
+
+			return Errors::AccessDenied();
+		}
+
+		return Errors::Success();
+	}
+
+	FOnlineError FSessionsCommon::CheckJoinSessionState(const FJoinSession::Params& Params)
+	{
+		// Check if a session with that name already exists
+		if (TSharedRef<FSession>* FoundSession = SessionsByName.Find(Params.SessionName))
+		{
+			UE_LOG(LogTemp, Warning, TEXT("[FSessionsCommon::CheckJoinSessionState] Could not join session with name [%s]. A session with that name already exists"), *Params.SessionName.ToString());
+
+			return Errors::InvalidState(); // TODO: New error: Session with name %s already exists
+		}
+
+		// User login check for all local users
+		IAuthPtr Auth = Services.GetAuthInterface();
+
+		TArray<FOnlineAccountIdHandle> LocalUserIds;
+		LocalUserIds.Reserve(Params.LocalUsers.Num());
+		Params.LocalUsers.GenerateKeyArray(LocalUserIds);
+
+		for (const FOnlineAccountIdHandle& LocalUserId : LocalUserIds)
+		{
+			FAuthGetAccountByAccountId::Params GetAccountByAccountIParams{ LocalUserId };
+			TOnlineResult<FAuthGetAccountByAccountId> Result = Auth->GetAccountByAccountId(MoveTemp(GetAccountByAccountIParams));
+
+			if (Result.IsOk())
+			{
+				if (Result.GetOkValue().AccountInfo->LoginStatus == ELoginStatus::NotLoggedIn)
+				{
+					UE_LOG(LogTemp, Warning, TEXT("[FSessionsCommon::CheckJoinSessionState] Could not join session with user [%s] not logged in"), *ToLogString(LocalUserId));
+
+					return Errors::InvalidUser();
+				}
+			}
+		}
+
+		return Errors::Success();
+	}
+
+	FOnlineError FSessionsCommon::CheckLeaveSessionState(const FLeaveSession::Params& Params)
+	{
+		// User login check for main caller, session check
+		if (TSharedRef<FSession>* FoundSession = SessionsByName.Find(Params.SessionName))
+		{
+			if (!FoundSession->Get().SessionSettings.IsDedicatedServerSession)
+			{
+				IAuthPtr Auth = Services.GetAuthInterface();
+
+				FAuthGetAccountByAccountId::Params GetAccountByAccountIParams{ Params.LocalUserId };
+				TOnlineResult<FAuthGetAccountByAccountId> Result = Auth->GetAccountByAccountId(MoveTemp(GetAccountByAccountIParams));
+
+				if (Result.IsOk())
+				{
+					if (Result.GetOkValue().AccountInfo->LoginStatus == ELoginStatus::NotLoggedIn)
+					{
+						UE_LOG(LogTemp, Warning, TEXT("[FSessionsNull::CheckLeaveSessionState] Could not leave session with user [%s] not logged in"), *ToLogString(Params.LocalUserId));
+
+						return Errors::InvalidUser();
+					}
+				}
+			}
+		}
+		else
+		{
+			UE_LOG(LogTemp, Warning, TEXT("[FSessionsCommon::CheckLeaveSessionState] Could not leave session with name [%s]. Session not found"), *Params.SessionName.ToString());
+
+			return Errors::NotFound();
+		}
+
+		// User login check for all local users
+		IAuthPtr Auth = Services.GetAuthInterface();
+
+		for (const FOnlineAccountIdHandle& LocalUserId : Params.LocalUsers)
+		{
+			FAuthGetAccountByAccountId::Params GetAccountByAccountIParams{ LocalUserId };
+			TOnlineResult<FAuthGetAccountByAccountId> Result = Auth->GetAccountByAccountId(MoveTemp(GetAccountByAccountIParams));
+
+			if (Result.IsOk())
+			{
+				if (Result.GetOkValue().AccountInfo->LoginStatus == ELoginStatus::NotLoggedIn)
+				{
+					UE_LOG(LogTemp, Warning, TEXT("[FSessionsCommon::CheckLeaveSessionState] Could not leave session with user [%s] not logged in"), *ToLogString(LocalUserId));
+
+					return Errors::InvalidUser();
+				}
+			}
+		}
+
+		return Errors::Success();
+	}
+
 /* UE::Online */ }
diff --git a/Engine/Plugins/Online/Experimental/OnlineServices/Source/OnlineServicesCommon/Public/Online/OnlineIdCommon.h b/Engine/Plugins/Online/Experimental/OnlineServices/Source/OnlineServicesCommon/Public/Online/OnlineIdCommon.h
--- a/Engine/Plugins/Online/Experimental/OnlineServices/Source/OnlineServicesCommon/Public/Online/OnlineIdCommon.h
+++ b/Engine/Plugins/Online/Experimental/OnlineServices/Source/OnlineServicesCommon/Public/Online/OnlineIdCommon.h
@@ -80,5 +80,8 @@ private:
 template<typename IdValueType, EOnlineServices OnlineServicesType>
 using TOnlineBasicAccountIdRegistry = TOnlineBasicIdRegistry<OnlineIdHandleTags::FAccount, IdValueType, OnlineServicesType>;
 
+template<typename IdValueType, EOnlineServices OnlineServicesType>
+using TOnlineBasicSessionIdRegistry = TOnlineBasicIdRegistry<OnlineIdHandleTags::FSession, IdValueType, OnlineServicesType>;
+
 /* UE::Online */ }
 
diff --git a/Engine/Plugins/Online/Experimental/OnlineServices/Source/OnlineServicesCommon/Public/Online/SessionsCommon.h b/Engine/Plugins/Online/Experimental/OnlineServices/Source/OnlineServicesCommon/Public/Online/SessionsCommon.h
--- a/Engine/Plugins/Online/Experimental/OnlineServices/Source/OnlineServicesCommon/Public/Online/SessionsCommon.h
+++ b/Engine/Plugins/Online/Experimental/OnlineServices/Source/OnlineServicesCommon/Public/Online/SessionsCommon.h
@@ -1,59 +1,123 @@
-// Copyright Epic Games, Inc. All Rights Reserved.
-
-#pragma once
-
-#include "Online/Sessions.h"
-#include "Online/OnlineComponent.h"
-
-namespace UE::Online {
-
-	class FOnlineServicesCommon;
-	class FAccountInfo;
-
-	struct FSessionEvents
-	{
-		TOnlineEventCallable<void(const FSessionJoined&)> OnSessionJoined;
-		TOnlineEventCallable<void(const FSessionLeft&)> OnSessionLeft;
-		TOnlineEventCallable<void(const FSessionUpdated&)> OnSessionUpdated;
-		TOnlineEventCallable<void(const FInviteReceived&)> OnInviteReceived;
-		TOnlineEventCallable<void(const FInviteAccepted&)> OnInviteAccepted;
-	};
-
-	class ONLINESERVICESCOMMON_API FSessionsCommon : public TOnlineComponent<ISessions>
-	{
-	public:
-		using Super = ISessions;
-
-		FSessionsCommon(FOnlineServicesCommon& InServices);
-
-		// TOnlineComponent
-		virtual void Initialize() override;
-		virtual void RegisterCommands() override;
-
-		// ISessions
-		virtual TOnlineResult<FGetAllSessions> GetAllSessions(FGetAllSessions::Params&& Params) override;
-		virtual TOnlineResult<FGetSessionByName> GetSessionByName(FGetSessionByName::Params&& Params) override;
-		virtual TOnlineResult<FGetSessionById> GetSessionById(FGetSessionById::Params&& Params) override;
-		virtual TOnlineAsyncOpHandle<FCreateSession> CreateSession(FCreateSession::Params&& Params) override;
-		virtual TOnlineAsyncOpHandle<FUpdateSession> UpdateSession(FUpdateSession::Params&& Params) override;
-		virtual TOnlineAsyncOpHandle<FLeaveSession> LeaveSession(FLeaveSession::Params&& Params) override;
-		virtual TOnlineAsyncOpHandle<FFindSessions> FindSessions(FFindSessions::Params&& Params) override;
-		virtual TOnlineAsyncOpHandle<FStartMatchmaking> StartMatchmaking(FStartMatchmaking::Params&& Params) override;
-		virtual TOnlineAsyncOpHandle<FJoinSession> JoinSession(FJoinSession::Params&& Params) override;
-		virtual TOnlineAsyncOpHandle<FSendSessionInvite> SendSessionInvite(FSendSessionInvite::Params&& Params) override;
-		virtual TOnlineAsyncOpHandle<FQuerySessionInvites> QuerySessionInvites(FQuerySessionInvites::Params&& Params) override;
-		virtual TOnlineAsyncOpHandle<FRejectSessionInvite> RejectSessionInvite(FRejectSessionInvite::Params&& Params) override;
-		virtual TOnlineAsyncOpHandle<FRegisterPlayers> RegisterPlayers(FRegisterPlayers::Params&& Params) override;
-		virtual TOnlineAsyncOpHandle<FUnregisterPlayers> UnregisterPlayers(FUnregisterPlayers::Params&& Params) override;
-
-		virtual TOnlineEvent<void(const FSessionJoined&)> OnSessionJoined() override;
-		virtual TOnlineEvent<void(const FSessionLeft&)> OnSessionLeft() override;
-		virtual TOnlineEvent<void(const FSessionUpdated&)> OnSessionUpdated() override;
-		virtual TOnlineEvent<void(const FInviteReceived&)> OnInviteReceived() override;
-		virtual TOnlineEvent<void(const FInviteAccepted&)> OnInviteAccepted() override;
-
-	protected:
-		FSessionEvents SessionEvents;
-};
-
-/* UE::Online */ }
+// Copyright Epic Games, Inc. All Rights Reserved.
+
+#pragma once
+
+#include "Online/Sessions.h"
+#include "Online/OnlineComponent.h"
+#include "Online/OnlineIdCommon.h"
+
+namespace UE::Online {
+
+class FOnlineServicesCommon;
+
+template<EOnlineServices OnlineServicesType>
+class TOnlineSessionIdStringRegistry : public IOnlineSessionIdRegistry
+{
+public:
+	// Begin IOnlineSessionIdRegistry
+	virtual inline FString ToLogString(const FOnlineSessionIdHandle& Handle) const override
+	{
+		FString IdValue = BasicRegistry.FindIdValue(Handle);
+
+		if (IdValue.Len() == 0)
+		{
+			IdValue = FString(TEXT("[InvalidSessionID]"));
+		}
+
+		return IdValue;
+	};
+
+	virtual inline TArray<uint8> ToReplicationData(const FOnlineSessionIdHandle& Handle) const override
+	{
+		const FTCHARToUTF8 IdValueUtf8(BasicRegistry.FindIdValue(Handle));
+
+		TArray<uint8> ReplicationData;
+		ReplicationData.SetNumUninitialized(IdValueUtf8.Length());
+		FMemory::Memcpy(ReplicationData.GetData(), IdValueUtf8.Get(), IdValueUtf8.Length());
+
+		return ReplicationData;
+	}
+
+	virtual inline FOnlineSessionIdHandle FromReplicationData(const TArray<uint8>& ReplicationData) override
+	{
+		FString Result = BytesToString(ReplicationData.GetData(), ReplicationData.Num());
+		if (Result.Len() > 0)
+		{
+			return BasicRegistry.FindOrAddHandle(Result);
+		}
+		return FOnlineSessionIdHandle();
+	}
+	// End IOnlineSessionIdRegistry
+
+	virtual ~TOnlineSessionIdStringRegistry() = default;
+
+public:
+	TOnlineBasicSessionIdRegistry<FString, OnlineServicesType> BasicRegistry;
+};
+
+class ONLINESERVICESCOMMON_API FSessionsCommon : public TOnlineComponent<ISessions>
+{
+public:
+	using Super = ISessions;
+
+	FSessionsCommon(FOnlineServicesCommon& InServices);
+
+	// TOnlineComponent
+	virtual void Initialize() override;
+	virtual void RegisterCommands() override;
+
+	// ISessions
+	virtual TOnlineResult<FGetAllSessions> GetAllSessions(FGetAllSessions::Params&& Params) const override;
+	virtual TOnlineResult<FGetSessionByName> GetSessionByName(FGetSessionByName::Params&& Params) const override;
+	virtual TOnlineResult<FGetSessionById> GetSessionById(FGetSessionById::Params&& Params) const override;
+	virtual TOnlineAsyncOpHandle<FCreateSession> CreateSession(FCreateSession::Params&& Params) override;
+	virtual TOnlineAsyncOpHandle<FUpdateSession> UpdateSession(FUpdateSession::Params&& Params) override;
+	virtual TOnlineAsyncOpHandle<FLeaveSession> LeaveSession(FLeaveSession::Params&& Params) override;
+	virtual TOnlineAsyncOpHandle<FFindSessions> FindSessions(FFindSessions::Params&& Params) override;
+	virtual TOnlineAsyncOpHandle<FStartMatchmaking> StartMatchmaking(FStartMatchmaking::Params&& Params) override;
+	virtual TOnlineAsyncOpHandle<FJoinSession> JoinSession(FJoinSession::Params&& Params) override;
+	virtual TOnlineAsyncOpHandle<FSendSessionInvite> SendSessionInvite(FSendSessionInvite::Params&& Params) override;
+	virtual TOnlineAsyncOpHandle<FQuerySessionInvites> QuerySessionInvites(FQuerySessionInvites::Params&& Params) override;
+	virtual TOnlineAsyncOpHandle<FRejectSessionInvite> RejectSessionInvite(FRejectSessionInvite::Params&& Params) override;
+	virtual TOnlineAsyncOpHandle<FRegisterPlayers> RegisterPlayers(FRegisterPlayers::Params&& Params) override;
+	virtual TOnlineAsyncOpHandle<FUnregisterPlayers> UnregisterPlayers(FUnregisterPlayers::Params&& Params) override;
+
+	virtual TOnlineEvent<void(const FSessionJoined&)> OnSessionJoined() override;
+	virtual TOnlineEvent<void(const FSessionLeft&)> OnSessionLeft() override;
+	virtual TOnlineEvent<void(const FSessionUpdated&)> OnSessionUpdated() override;
+	virtual TOnlineEvent<void(const FInviteReceived&)> OnInviteReceived() override;
+	virtual TOnlineEvent<void(const FInviteAccepted&)> OnInviteAccepted() override;
+
+protected:
+	FOnlineError CheckCreateSessionParams(const FCreateSession::Params& Params);
+	FOnlineError CheckCreateSessionState(const FCreateSession::Params& Params);
+
+	FOnlineError CheckUpdateSessionState(const FUpdateSession::Params& Params);
+
+	FOnlineError CheckFindSessionsParams(const FFindSessions::Params& Params);
+	FOnlineError CheckFindSessionsState(const FFindSessions::Params& Params);
+
+	FOnlineError CheckJoinSessionParams(const FJoinSession::Params& Params);
+	FOnlineError CheckJoinSessionState(const FJoinSession::Params& Params);
+
+	FOnlineError CheckLeaveSessionState(const FLeaveSession::Params& Params);
+
+protected:
+
+	struct FSessionEvents
+	{
+		TOnlineEventCallable<void(const FSessionJoined&)> OnSessionJoined;
+		TOnlineEventCallable<void(const FSessionLeft&)> OnSessionLeft;
+		TOnlineEventCallable<void(const FSessionUpdated&)> OnSessionUpdated;
+		TOnlineEventCallable<void(const FInviteReceived&)> OnInviteReceived;
+		TOnlineEventCallable<void(const FInviteAccepted&)> OnInviteAccepted;
+	} SessionEvents;
+
+	TMap<FName, TSharedRef<FSession>> SessionsByName;
+	TMap<FOnlineSessionIdHandle, TSharedRef<FSession>> SessionsById;
+
+	TSharedPtr<FFindSessions::Result> CurrentSessionSearch;
+	TSharedPtr<TOnlineAsyncOp<FFindSessions>> CurrentSessionSearchHandle;
+};
+
+/* UE::Online */ }
diff --git a/Engine/Plugins/Online/Experimental/OnlineServices/Source/OnlineServicesInterface/Public/Online/Sessions.h b/Engine/Plugins/Online/Experimental/OnlineServices/Source/OnlineServicesInterface/Public/Online/Sessions.h
--- a/Engine/Plugins/Online/Experimental/OnlineServices/Source/OnlineServicesInterface/Public/Online/Sessions.h
+++ b/Engine/Plugins/Online/Experimental/OnlineServices/Source/OnlineServicesInterface/Public/Online/Sessions.h
@@ -125,16 +125,16 @@ struct FSessionSettings
 	FName SchemaName;
 
 	/* Maximum number of public slots for session members */
-	int32 NumMaxPublicConnections;
+	uint32 NumMaxPublicConnections;
 
 	/* Number of available public slots for session members */
-	int32 NumOpenPublicConnections;
+	uint32 NumOpenPublicConnections;
 
 	/* Maximum number of private slots for session members */
-	int32 NumMaxPrivateConnections;
+	uint32 NumMaxPrivateConnections;
 
 	/* Number of available private slots for session members */
-	int32 NumOpenPrivateConnections;
+	uint32 NumOpenPrivateConnections;
 
 	/* Enum value describing the level of restriction to join the session */
 	ESessionJoinPolicy JoinPolicy;
@@ -160,8 +160,8 @@ struct FSessionSettings
 	/*Whether this is a secure session protected by anti-cheat services */
 	bool bAntiCheatProtected;
 
-	 /* Whether this session will show its information in presence updates. Can only be set in one session at a time */
-	 bool bPresenceEnabled;
+	/* Whether this session will show its information in presence updates. Can only be set in one session at a time */
+	bool bPresenceEnabled;
   
  	/* Map of user-defined settings to be passed to the platform APIs as additional information for various purposes */
  	FCustomSessionSettingsMap CustomSettings;
@@ -176,10 +176,10 @@ struct FSessionSettings
 struct FSessionSettingsUpdate
 {
 	TOptional<FName> SchemaName;
-	TOptional<int32> NumMaxPublicConnections;
-	TOptional<int32> NumOpenPublicConnections;
-	TOptional<int32> NumMaxPrivateConnections;
-	TOptional<int32> NumOpenPrivateConnections;
+	TOptional<uint32> NumMaxPublicConnections;
+	TOptional<uint32> NumOpenPublicConnections;
+	TOptional<uint32> NumMaxPrivateConnections;
+	TOptional<uint32> NumOpenPrivateConnections;
 	TOptional<ESessionJoinPolicy> JoinPolicy;
 	TOptional<FString> SessionIdOverride;
 	TOptional<bool> IsLANSession;
@@ -216,7 +216,7 @@ struct ONLINESERVICESINTERFACE_API FSession
 	/** The current state of the session */
 	ESessionState CurrentState = ESessionState::Invalid;
 
-	/** Set of session properties */
+	/** Set of session properties that can be altered by the session owner */
 	FSessionSettings SessionSettings;
 };
 
@@ -344,8 +344,8 @@ struct FLeaveSession
 		/* The local name for the session. */
 		FName SessionName;
 
-		/** Ids for local users (other than the main caller) who will also be leaving the session */
-		TArray<FOnlineAccountIdHandle> AdditionalLocalUsers;
+		/** Ids for all local users who will leave the session (includes the main caller) */
+		TArray<FOnlineAccountIdHandle> LocalUsers;
 
 		/* Whether the call should attempt to destroy the session instead of just leave it */
 		bool bDestroySession;
@@ -594,7 +594,7 @@ public:
 	 * @params Parameters for the GetAllSessions call
 	 * return
 	 */
-	virtual TOnlineResult<FGetAllSessions> GetAllSessions(FGetAllSessions::Params&& Params) = 0;
+	virtual TOnlineResult<FGetAllSessions> GetAllSessions(FGetAllSessions::Params&& Params) const = 0;
 
 	/**
 	 * Get the session object with a given local name.
@@ -602,7 +602,7 @@ public:
 	 * @params Parameters for the GetSessionByName call
 	 * return
 	 */
-	virtual TOnlineResult<FGetSessionByName> GetSessionByName(FGetSessionByName::Params&& Params) = 0;
+	virtual TOnlineResult<FGetSessionByName> GetSessionByName(FGetSessionByName::Params&& Params) const = 0;
 
 	/**
 	 * Get the session object with a given id handle.
@@ -610,7 +610,7 @@ public:
 	 * @params Parameters for the GetSessionById call
 	 * return
 	 */
-	virtual TOnlineResult<FGetSessionById> GetSessionById(FGetSessionById::Params&& Params) = 0;
+	virtual TOnlineResult<FGetSessionById> GetSessionById(FGetSessionById::Params&& Params) const = 0;
 
 	/**
 	 * Create and join a new session.
@@ -868,7 +868,7 @@ END_ONLINE_STRUCT_META()
 BEGIN_ONLINE_STRUCT_META(FLeaveSession::Params)
 	ONLINE_STRUCT_FIELD(FLeaveSession::Params, LocalUserId),
 	ONLINE_STRUCT_FIELD(FLeaveSession::Params, SessionName),
-	ONLINE_STRUCT_FIELD(FLeaveSession::Params, AdditionalLocalUsers),
+	ONLINE_STRUCT_FIELD(FLeaveSession::Params, LocalUsers),
 	ONLINE_STRUCT_FIELD(FLeaveSession::Params, bDestroySession)
 END_ONLINE_STRUCT_META()
 
diff --git a/Engine/Plugins/Online/Experimental/OnlineServicesEOSGS/Source/Private/Online/SessionsEOSGS.cpp b/Engine/Plugins/Online/Experimental/OnlineServicesEOSGS/Source/Private/Online/SessionsEOSGS.cpp
--- a/Engine/Plugins/Online/Experimental/OnlineServicesEOSGS/Source/Private/Online/SessionsEOSGS.cpp
+++ b/Engine/Plugins/Online/Experimental/OnlineServicesEOSGS/Source/Private/Online/SessionsEOSGS.cpp
@@ -35,74 +35,6 @@ FOnlineSessionIdRegistryEOSGS& FOnlineSessionIdRegistryEOSGS::Get()
 	return Instance;
 }
 
-FOnlineSessionIdHandle FOnlineSessionIdRegistryEOSGS::Find(const FString& SessionId) const
-{
-	if (const FOnlineSessionIdString* Entry = StringToId.Find(SessionId))
-	{
-		return Entry->Handle;
-	}
-	return FOnlineSessionIdHandle();
-}
-
-FOnlineSessionIdHandle FOnlineSessionIdRegistryEOSGS::FindOrAdd(const FString& SessionId)
-{
-	const FOnlineSessionIdString* Entry = StringToId.Find(SessionId);
-	if (Entry)
-	{
-		return Entry->Handle;
-	}
-
-	FOnlineSessionIdString& NewId = Ids.Emplace_GetRef();
-	NewId.Data = SessionId;
-	NewId.Handle = FOnlineSessionIdHandle(EOnlineServices::Epic, Ids.Num());
-
-	StringToId.Add(SessionId, NewId);
-
-	return NewId.Handle;
-}
-
-const FOnlineSessionIdString* FOnlineSessionIdRegistryEOSGS::GetInternal(const FOnlineSessionIdHandle& Handle) const
-{
-	if (Handle.IsValid() && Handle.GetOnlineServicesType() == EOnlineServices::Epic && Handle.GetHandle() <= (uint32)Ids.Num())
-	{
-		return &Ids[Handle.GetHandle() - 1];
-	}
-	return nullptr;
-}
-
-FString FOnlineSessionIdRegistryEOSGS::ToLogString(const FOnlineSessionIdHandle& Handle) const
-{
-	if (const FOnlineSessionIdString* Id = GetInternal(Handle))
-	{
-		return Id->Data;
-	}
-
-	return FString(TEXT("[InvalidSessionID]"));
-}
-
-TArray<uint8> FOnlineSessionIdRegistryEOSGS::ToReplicationData(const FOnlineSessionIdHandle& Handle) const
-{
-	if (const FOnlineSessionIdString* Id = GetInternal(Handle))
-	{
-		TArray<uint8> ReplicationData;
-		ReplicationData.SetNumUninitialized(Id->Data.Len());
-		StringToBytes(Id->Data, ReplicationData.GetData(), Id->Data.Len());
-		return ReplicationData;
-	}
-
-	return TArray<uint8>();;
-}
-
-FOnlineSessionIdHandle FOnlineSessionIdRegistryEOSGS::FromReplicationData(const TArray<uint8>& ReplicationData)
-{
-	FString Result = BytesToString(ReplicationData.GetData(), ReplicationData.Num());
-	if (Result.Len() > 0)
-	{
-		return FindOrAdd(Result);
-	}
-	return FOnlineSessionIdHandle();
-}
-
 /** FSessionEOSGS */
 
 FSessionEOSGS::FSessionEOSGS()
@@ -130,7 +62,7 @@ FSessionEOSGS::FSessionEOSGS(const EOS_HSessionDetails& InSessionDetailsHandle)
 	if (CopyInfoResult == EOS_EResult::EOS_Success)
 	{
 		CurrentState = ESessionState::Valid;
-		SessionId = FOnlineSessionIdRegistryEOSGS::Get().FindOrAdd(FString(SessionDetailsInfo->SessionId));
+		SessionId = FSessionsEOSGS::CreateSessionId(FString(SessionDetailsInfo->SessionId));
 
 		// TODO: Load all session settings after SchemaVariant work
 	}
@@ -165,42 +97,6 @@ void FSessionsEOSGS::Initialize()
 	check(SessionsHandle);
 }
 
-TOnlineResult<FGetAllSessions> FSessionsEOSGS::GetAllSessions(FGetAllSessions::Params&& Params)
-{
-	FGetAllSessions::Result Result;
-
-	for (TMap<FName, TSharedRef<FSessionEOSGS>>::TConstIterator It(SessionsByName); It; ++It)
-	{
-		Result.Sessions.Add(It.Value());
-	}
-
-	return TOnlineResult<FGetAllSessions>(MoveTemp(Result));
-}
-
-TOnlineResult<FGetSessionByName> FSessionsEOSGS::GetSessionByName(FGetSessionByName::Params&& Params)
-{
-	if (TSharedRef<FSessionEOSGS>* FoundSession = SessionsByName.Find(Params.LocalName))
-	{
-		return TOnlineResult<FGetSessionByName>({ *FoundSession });
-	}
-	else
-	{
-		return TOnlineResult<FGetSessionByName>(Errors::NotFound());
-	}
-}
-
-TOnlineResult<FGetSessionById> FSessionsEOSGS::GetSessionById(FGetSessionById::Params&& Params)
-{
-	if (TSharedRef<FSessionEOSGS>* FoundSession = SessionsById.Find(Params.IdHandle))
-	{
-		return TOnlineResult<FGetSessionById>({ *FoundSession });
-	}
-	else
-	{
-		return TOnlineResult<FGetSessionById>(Errors::NotFound());
-	}
-}
-
 TOnlineAsyncOpHandle<FCreateSession> FSessionsEOSGS::CreateSession(FCreateSession::Params&& Params)
 {
 	TOnlineAsyncOpRef<FCreateSession> Op = GetOp<FCreateSession>(MoveTemp(Params));
@@ -218,7 +114,7 @@ TOnlineAsyncOpHandle<FCreateSession> FSessionsEOSGS::CreateSession(FCreateSessio
 			return MakeFulfilledPromise<TDefaultErrorResult<FUpdateSessionImpl>>().GetFuture();
 		}
 
-		if (OpParams.SessionSettings.NumMaxPrivateConnections <= 0 && OpParams.SessionSettings.NumMaxPublicConnections <= 0)
+		if (OpParams.SessionSettings.NumMaxPrivateConnections == 0 && OpParams.SessionSettings.NumMaxPublicConnections == 0)
 		{
 			UE_LOG(LogTemp, Warning, TEXT("[FSessionsEOSGS::CreateSession] Could not create session with no valid NumMaxPrivateConnections [%d] or NumMaxPublicConnections [%d]"), OpParams.SessionSettings.NumMaxPrivateConnections, OpParams.SessionSettings.NumMaxPublicConnections);
 
@@ -390,7 +286,7 @@ void FSessionsEOSGS::SetBucketId(EOS_HSessionModification& SessionModificationHa
 	}
 }
 
-void FSessionsEOSGS::SetMaxPlayers(EOS_HSessionModification& SessionModificationHandle, const int32& NewMaxPlayers)
+void FSessionsEOSGS::SetMaxPlayers(EOS_HSessionModification& SessionModificationHandle, const uint32& NewMaxPlayers)
 {
 	EOS_SessionModification_SetMaxPlayersOptions Options = { };
 	Options.ApiVersion = EOS_SESSIONMODIFICATION_SETMAXPLAYERS_API_LATEST;
@@ -755,7 +651,7 @@ TFuture<TDefaultErrorResult<FUpdateSessionImpl>> FSessionsEOSGS::UpdateSessionIm
 			// In session creation calls, we'll need to set the new id for the session
 			if (!FoundSession->Get().SessionId.IsValid())
 			{
-				FoundSession->Get().SessionId = FOnlineSessionIdRegistryEOSGS::Get().FindOrAdd(FString(Result->SessionId));
+				FoundSession->Get().SessionId = CreateSessionId(FString(Result->SessionId));
 			}
 		}
 
@@ -880,6 +776,11 @@ void FSessionsEOSGS::WriteSessionSearchHandle(FSessionSearchHandleEOSGS& Session
 	}
 }
 
+FOnlineSessionIdHandle FSessionsEOSGS::CreateSessionId(const FString& SessionId)
+{
+	return FOnlineSessionIdRegistryEOSGS::Get().BasicRegistry.FindOrAddHandle(SessionId);
+}
+
 TOnlineAsyncOpHandle<FFindSessions> FSessionsEOSGS::FindSessions(FFindSessions::Params&& Params)
 {
 	TOnlineAsyncOpRef<FFindSessions> Op = GetOp<FFindSessions>(MoveTemp(Params));
diff --git a/Engine/Plugins/Online/Experimental/OnlineServicesEOSGS/Source/Public/Online/SessionsEOSGS.h b/Engine/Plugins/Online/Experimental/OnlineServicesEOSGS/Source/Public/Online/SessionsEOSGS.h
--- a/Engine/Plugins/Online/Experimental/OnlineServicesEOSGS/Source/Public/Online/SessionsEOSGS.h
+++ b/Engine/Plugins/Online/Experimental/OnlineServicesEOSGS/Source/Public/Online/SessionsEOSGS.h
@@ -15,37 +15,10 @@ namespace UE::Online {
 
 class FOnlineServicesEOSGS;
 
-/** Session Id Registry */
-
-// TODO: Move to Common since it's in Null too
-class FOnlineSessionIdString
-{
-public:
-	FString Data;
-	FOnlineSessionIdHandle Handle;
-};
-
-class FOnlineSessionIdRegistryEOSGS : public IOnlineSessionIdRegistry
+class FOnlineSessionIdRegistryEOSGS : public TOnlineSessionIdStringRegistry<EOnlineServices::Epic>
 {
 public:
 	static FOnlineSessionIdRegistryEOSGS& Get();
-
-	FOnlineSessionIdHandle Find(const FString& SessionId) const;
-	FOnlineSessionIdHandle FindOrAdd(const FString& SessionId);
-
-	// Begin IOnlineSessionIdRegistry
-	virtual FString ToLogString(const FOnlineSessionIdHandle& Handle) const override;
-	virtual TArray<uint8> ToReplicationData(const FOnlineSessionIdHandle& Handle) const override;
-	virtual FOnlineSessionIdHandle FromReplicationData(const TArray<uint8>& ReplicationString) override;
-	// End IOnlineSessionIdRegistry
-
-	virtual ~FOnlineSessionIdRegistryEOSGS() = default;
-
-private:
-	const FOnlineSessionIdString* GetInternal(const FOnlineSessionIdHandle& Handle) const;
-
-	TArray<FOnlineSessionIdString> Ids;
-	TMap<FString, FOnlineSessionIdString> StringToId;
 };
 
 static FName EOS_SESSIONS_BUCKET_ID = TEXT("EOS_SESSIONS_BUCKET_ID");
@@ -112,6 +85,8 @@ struct FUpdateSessionImpl
 class ONLINESERVICESEOSGS_API FSessionsEOSGS : public FSessionsCommon
 {
 public:
+	friend class FSessionEOSGS;
+
 	using Super = FSessionsCommon;
 
 	FSessionsEOSGS(FOnlineServicesEOSGS& InOwningSubsystem);
@@ -121,9 +96,6 @@ public:
 	void Initialize() override;
 
 	// ISessions
-	virtual TOnlineResult<FGetAllSessions> GetAllSessions(FGetAllSessions::Params&& Params) override;
-	virtual TOnlineResult<FGetSessionByName> GetSessionByName(FGetSessionByName::Params&& Params) override;
-	virtual TOnlineResult<FGetSessionById> GetSessionById(FGetSessionById::Params&& Params) override;
 	virtual TOnlineAsyncOpHandle<FCreateSession> CreateSession(FCreateSession::Params&& Params) override;
 	virtual TOnlineAsyncOpHandle<FUpdateSession> UpdateSession(FUpdateSession::Params&& Params) override;
 	virtual TOnlineAsyncOpHandle<FLeaveSession> LeaveSession(FLeaveSession::Params&& Params) override;
@@ -133,7 +105,7 @@ public:
 protected:
 	void SetPermissionLevel(EOS_HSessionModification& SessionModHandle, const ESessionJoinPolicy& JoinPolicy);
 	void SetBucketId(EOS_HSessionModification& SessionModHandle, const FString& NewBucketId);
-	void SetMaxPlayers(EOS_HSessionModification& SessionModHandle, const int32& NewMaxPlayers);
+	void SetMaxPlayers(EOS_HSessionModification& SessionModHandle, const uint32& NewMaxPlayers);
 
 	/**
 	 * Writes all values in the passed SessionSettings to the SessionModificationHandle
@@ -147,6 +119,8 @@ protected:
 
 	void WriteSessionSearchHandle(FSessionSearchHandleEOSGS& SessionSearchHandle, const FFindSessions::Params& Params);
 
+	static FOnlineSessionIdHandle CreateSessionId(const FString& SessionId);
+
 protected:
 	FOnlineServicesEOSGS& Services;
 
diff --git a/Engine/Plugins/Online/Experimental/OnlineServicesNull/Source/Private/Online/SessionsNull.cpp b/Engine/Plugins/Online/Experimental/OnlineServicesNull/Source/Private/Online/SessionsNull.cpp
--- a/Engine/Plugins/Online/Experimental/OnlineServicesNull/Source/Private/Online/SessionsNull.cpp
+++ b/Engine/Plugins/Online/Experimental/OnlineServicesNull/Source/Private/Online/SessionsNull.cpp
@@ -21,88 +21,27 @@ FOnlineSessionIdRegistryNull& FOnlineSessionIdRegistryNull::Get()
 	return Instance;
 }
 
-FOnlineSessionIdHandle FOnlineSessionIdRegistryNull::Find(FString SessionId) const
+FOnlineSessionIdHandle FOnlineSessionIdRegistryNull::GetNextSessionId()
 {
-	if(const FOnlineSessionIdString* Entry = StringToId.Find(SessionId))
-	{
-		return Entry->Handle;
-	}
-	return FOnlineSessionIdHandle();
-}
-
-FOnlineSessionIdHandle FOnlineSessionIdRegistryNull::FindOrAdd(FString SessionId)
-{
-	const FOnlineSessionIdString* Entry = StringToId.Find(SessionId);
-	if (Entry)
-	{
-		return Entry->Handle;
-	}
-
-	FOnlineSessionIdString& NewId = Ids.Emplace_GetRef();
-	NewId.Data = SessionId;
-	NewId.Handle = FOnlineSessionIdHandle(EOnlineServices::Null, Ids.Num());
-
-	StringToId.Add(SessionId, NewId);
-
-	return NewId.Handle;
-}
-
-FOnlineSessionIdHandle FOnlineSessionIdRegistryNull::GetNext()
-{
-	return FindOrAdd(FGuid().ToString());
-}
-
-const FOnlineSessionIdString* FOnlineSessionIdRegistryNull::GetInternal(const FOnlineSessionIdHandle& Handle) const
-{
-	if (Handle.IsValid() && Handle.GetOnlineServicesType() == EOnlineServices::Null && Handle.GetHandle() <= (uint32)Ids.Num())
-	{
-		return &Ids[Handle.GetHandle() - 1];
-	}
-	return nullptr;
-}
-
-FString FOnlineSessionIdRegistryNull::ToLogString(const FOnlineSessionIdHandle& Handle) const
-{
-	if (const FOnlineSessionIdString* Id = GetInternal(Handle))
-	{
-		return Id->Data;
-	}
-
-	return FString(TEXT("[InvalidSessionID]"));
-}
-
-TArray<uint8> FOnlineSessionIdRegistryNull::ToReplicationData(const FOnlineSessionIdHandle& Handle) const
-{
-	if (const FOnlineSessionIdString* Id = GetInternal(Handle))
-	{
-		TArray<uint8> ReplicationData;
-		ReplicationData.SetNumUninitialized(Id->Data.Len());
-		StringToBytes(Id->Data, ReplicationData.GetData(), Id->Data.Len());
-		return ReplicationData;
-	}
-
-	return TArray<uint8>();;
-}
-
-FOnlineSessionIdHandle FOnlineSessionIdRegistryNull::FromReplicationData(const TArray<uint8>& ReplicationData)
-{
-	FString Result = BytesToString(ReplicationData.GetData(), ReplicationData.Num());
-	if (Result.Len() > 0)
-	{
-		return FindOrAdd(Result);
-	}
-	return FOnlineSessionIdHandle();
+	return BasicRegistry.FindOrAddHandle(FGuid().ToString());
 }
 
 /** FSessionNull */
 
 FSessionNull::FSessionNull()
-	: FSessionNull(FSession())
 {
+	Initialize();
 }
 
 FSessionNull::FSessionNull(const FSession& InSession)
 	: FSession(InSession)
+{
+	const FSessionNull& InSessionNull = FSessionNull::Cast(InSession);
+
+	OwnerInternetAddr = InSessionNull.OwnerInternetAddr;
+}
+
+void FSessionNull::Initialize()
 {
 	bool bCanBindAll;
 	OwnerInternetAddr = ISocketSubsystem::Get(PLATFORM_SOCKETSUBSYSTEM)->GetLocalHostAddr(*GLog, bCanBindAll);
@@ -132,11 +71,24 @@ FSessionNull::FSessionNull(const FSession& InSession)
 	}
 }
 
+FSessionNull& FSessionNull::Cast(FSession& InSession)
+{
+	check(InSession.SessionId.GetOnlineServicesType() == EOnlineServices::Null);
+	
+	return *static_cast<FSessionNull*>(&InSession);
+}
+
+const FSessionNull& FSessionNull::Cast(const FSession& InSession)
+{
+	check(InSession.SessionId.GetOnlineServicesType() == EOnlineServices::Null);
+
+	return *static_cast<const FSessionNull*>(&InSession);
+}
+
 /** FSessionsNull */
 
 FSessionsNull::FSessionsNull(FOnlineServicesNull& InServices)
 	: Super(InServices)
-	, Services(InServices)
 {
 	if (!LANSessionManager.IsValid())
 	{
@@ -154,60 +106,15 @@ void FSessionsNull::Tick(float DeltaSeconds)
 	Super::Tick(DeltaSeconds);
 }
 
-TOnlineResult<FGetAllSessions> FSessionsNull::GetAllSessions(FGetAllSessions::Params&& Params)
-{
-	FGetAllSessions::Result Result;
-
-	for (TMap<FName, TSharedRef<FSessionNull>>::TConstIterator It(SessionsByName); It; ++It)
-	{
-		Result.Sessions.Add(It.Value());
-	}
-
-	return TOnlineResult<FGetAllSessions>(MoveTemp(Result));
-}
-
-TOnlineResult<FGetSessionByName> FSessionsNull::GetSessionByName(FGetSessionByName::Params&& Params)
-{
-	if (TSharedRef<FSessionNull>* FoundSession = SessionsByName.Find(Params.LocalName))
-	{
-		return TOnlineResult<FGetSessionByName>({ *FoundSession });
-	}
-	else
-	{
-		return TOnlineResult<FGetSessionByName>(Errors::NotFound());
-	}
-}
-
-TOnlineResult<FGetSessionById> FSessionsNull::GetSessionById(FGetSessionById::Params&& Params)
-{
-	if (TSharedRef<FSessionNull>* FoundSession = SessionsById.Find(Params.IdHandle))
-	{
-		return TOnlineResult<FGetSessionById>({ *FoundSession });
-	}
-	else
-	{
-		return TOnlineResult<FGetSessionById>(Errors::NotFound());
-	}
-}
-
 TOnlineAsyncOpHandle<FCreateSession> FSessionsNull::CreateSession(FCreateSession::Params&& Params)
 {
 	TOnlineAsyncOpRef<FCreateSession> Op = GetOp<FCreateSession>(MoveTemp(Params));
 	const FCreateSession::Params& OpParams = Op->GetParams();
 
-	if (TSharedRef<FSessionNull>* FoundSession = SessionsByName.Find(OpParams.SessionName))
-	{
-		UE_LOG(LogTemp, Warning, TEXT("[FSessionsNull::CreateSession] Could not create session with name [%s]. A session with that name already exists"), *OpParams.SessionName.ToString());
-
-		Op->SetError(Errors::InvalidParams()); // TODO: New error: Session with name %s already exists
-		return Op->GetHandle();
-	}
-
-	if (OpParams.SessionSettings.NumMaxPrivateConnections <= 0 && OpParams.SessionSettings.NumMaxPublicConnections <= 0)
+	FOnlineError ParamsCheck = CheckCreateSessionParams(OpParams);
+	if (ParamsCheck != Errors::Success())
 	{
-		UE_LOG(LogTemp, Warning, TEXT("[FSessionsNull::CreateSession] Could not create session with no valid NumMaxPrivateConnections [%d] or NumMaxPublicConnections [%d]"), OpParams.SessionSettings.NumMaxPrivateConnections, OpParams.SessionSettings.NumMaxPublicConnections);
-
-		Op->SetError(Errors::InvalidParams());
+		Op->SetError(MoveTemp(ParamsCheck));
 		return Op->GetHandle();
 	}
 
@@ -225,10 +132,17 @@ TOnlineAsyncOpHandle<FCreateSession> FSessionsNull::CreateSession(FCreateSession
 	{
 		const FCreateSession::Params& OpParams = Op.GetParams();
 
+		FOnlineError StateCheck = CheckCreateSessionState(OpParams);
+		if (StateCheck != Errors::Success())
+		{
+			Op.SetError(MoveTemp(StateCheck));
+			return;
+		}
+
 		FSessionNull NewSession = FSessionNull();
 		NewSession.CurrentState = ESessionState::Valid;
 		NewSession.OwnerUserId = OpParams.LocalUserId;
-		NewSession.SessionId = FOnlineSessionIdRegistryNull::Get().GetNext();
+		NewSession.SessionId = FOnlineSessionIdRegistryNull::Get().GetNextSessionId();
 		NewSession.SessionSettings = OpParams.SessionSettings;
 
 		// For Null sessions, we'll add all the Session members manually instead of calling JoinSession since there is no API calls involved
@@ -254,139 +168,140 @@ TOnlineAsyncOpHandle<FUpdateSession> FSessionsNull::UpdateSession(FUpdateSession
 	{
 		const FUpdateSession::Params& OpParams = Op.GetParams();
 
-		if (TSharedRef<FSessionNull>* FoundSession = SessionsByName.Find(OpParams.SessionName))
+		FOnlineError StateCheck = CheckUpdateSessionState(OpParams);
+		if (StateCheck != Errors::Success())
 		{
-			// If the session exists, we'll check and update all settings one by one, with additional logic wherever required
-
-			FSessionSettings& SessionSettings = FoundSession->Get().SessionSettings;
-			const FSessionSettingsUpdate& UpdatedSettings = OpParams.Mutations;
+			Op.SetError(MoveTemp(StateCheck));
+			return;
+		}
 
-			if (UpdatedSettings.bAllowNewMembers.IsSet())
-			{
-				SessionSettings.bAllowNewMembers = UpdatedSettings.bAllowNewMembers.GetValue();
-			}
+		// We'll check and update all settings one by one, with additional logic wherever required
 
-			if (UpdatedSettings.bAllowSanctionedPlayers.IsSet())
-			{
-				SessionSettings.bAllowSanctionedPlayers = UpdatedSettings.bAllowSanctionedPlayers.GetValue();
-			}
+		TSharedRef<FSession>& FoundSession = SessionsByName.FindChecked(OpParams.SessionName);
+		FSessionSettings& SessionSettings = FoundSession->SessionSettings;
+		const FSessionSettingsUpdate& UpdatedSettings = OpParams.Mutations;
 
-			if (UpdatedSettings.bAllowUnregisteredPlayers.IsSet())
-			{
-				SessionSettings.bAllowUnregisteredPlayers = UpdatedSettings.bAllowUnregisteredPlayers.GetValue();
-			}
+		if (UpdatedSettings.bAllowNewMembers.IsSet())
+		{
+			SessionSettings.bAllowNewMembers = UpdatedSettings.bAllowNewMembers.GetValue();
+		}
 
-			if (UpdatedSettings.bAntiCheatProtected.IsSet())
-			{
-				SessionSettings.bAntiCheatProtected = UpdatedSettings.bAntiCheatProtected.GetValue();
-			}
+		if (UpdatedSettings.bAllowSanctionedPlayers.IsSet())
+		{
+			SessionSettings.bAllowSanctionedPlayers = UpdatedSettings.bAllowSanctionedPlayers.GetValue();
+		}
 
-			if (UpdatedSettings.bPresenceEnabled.IsSet())
-			{
-				SessionSettings.bPresenceEnabled = UpdatedSettings.bPresenceEnabled.GetValue();
-			}
+		if (UpdatedSettings.bAllowUnregisteredPlayers.IsSet())
+		{
+			SessionSettings.bAllowUnregisteredPlayers = UpdatedSettings.bAllowUnregisteredPlayers.GetValue();
+		}
 
-			if (UpdatedSettings.IsDedicatedServerSession.IsSet())
-			{
-				SessionSettings.IsDedicatedServerSession = UpdatedSettings.IsDedicatedServerSession.GetValue();
-			}
+		if (UpdatedSettings.bAntiCheatProtected.IsSet())
+		{
+			SessionSettings.bAntiCheatProtected = UpdatedSettings.bAntiCheatProtected.GetValue();
+		}
 
-			if (UpdatedSettings.IsLANSession.IsSet())
-			{
-				SessionSettings.IsLANSession = UpdatedSettings.IsLANSession.GetValue();
-			}
+		if (UpdatedSettings.bPresenceEnabled.IsSet())
+		{
+			SessionSettings.bPresenceEnabled = UpdatedSettings.bPresenceEnabled.GetValue();
+		}
 
-			if (UpdatedSettings.JoinPolicy.IsSet())
-			{
-				// Changes in the join policy setting will mean we start or stop the LAN Beacon broadcasting the session information
-				// There is no FriendsNull interface at the time of this implementation, hence the binary behavior (Public/Non-Public)
-				if (SessionSettings.JoinPolicy != ESessionJoinPolicy::Public && UpdatedSettings.JoinPolicy == ESessionJoinPolicy::Public)
-				{
-					TryHostLANSession();
-				}
-				else if (SessionSettings.JoinPolicy == ESessionJoinPolicy::Public && UpdatedSettings.JoinPolicy != ESessionJoinPolicy::Public)
-				{
-					StopLANSession();
-				}
+		if (UpdatedSettings.IsDedicatedServerSession.IsSet())
+		{
+			SessionSettings.IsDedicatedServerSession = UpdatedSettings.IsDedicatedServerSession.GetValue();
+		}
 
-				SessionSettings.JoinPolicy = UpdatedSettings.JoinPolicy.GetValue();
-			}
+		if (UpdatedSettings.IsLANSession.IsSet())
+		{
+			SessionSettings.IsLANSession = UpdatedSettings.IsLANSession.GetValue();
+		}
 
-			if (UpdatedSettings.NumMaxPrivateConnections.IsSet())
+		if (UpdatedSettings.JoinPolicy.IsSet())
+		{
+			// Changes in the join policy setting will mean we start or stop the LAN Beacon broadcasting the session information
+			// There is no FriendsNull interface at the time of this implementation, hence the binary behavior (Public/Non-Public)
+			if (SessionSettings.JoinPolicy != ESessionJoinPolicy::Public && UpdatedSettings.JoinPolicy == ESessionJoinPolicy::Public)
 			{
-				SessionSettings.NumMaxPrivateConnections = UpdatedSettings.NumMaxPrivateConnections.GetValue();
+				TryHostLANSession();
 			}
-
-			if (UpdatedSettings.NumMaxPublicConnections.IsSet())
+			else if (SessionSettings.JoinPolicy == ESessionJoinPolicy::Public && UpdatedSettings.JoinPolicy != ESessionJoinPolicy::Public)
 			{
-				SessionSettings.NumMaxPublicConnections = UpdatedSettings.NumMaxPublicConnections.GetValue();
+				StopLANSession();
 			}
 
-			if (UpdatedSettings.NumOpenPrivateConnections.IsSet())
-			{
-				SessionSettings.NumOpenPrivateConnections = UpdatedSettings.NumOpenPrivateConnections.GetValue();
-			}
+			SessionSettings.JoinPolicy = UpdatedSettings.JoinPolicy.GetValue();
+		}
 
-			if (UpdatedSettings.NumOpenPublicConnections.IsSet())
-			{
-				SessionSettings.NumOpenPublicConnections = UpdatedSettings.NumOpenPublicConnections.GetValue();
-			}
+		if (UpdatedSettings.NumMaxPrivateConnections.IsSet())
+		{
+			SessionSettings.NumMaxPrivateConnections = UpdatedSettings.NumMaxPrivateConnections.GetValue();
+		}
 
-			// TODO: We may need some additional logic for schema changes
-			if (UpdatedSettings.SchemaName.IsSet())
-			{
-				SessionSettings.SchemaName = UpdatedSettings.SchemaName.GetValue();
-			}
+		if (UpdatedSettings.NumMaxPublicConnections.IsSet())
+		{
+			SessionSettings.NumMaxPublicConnections = UpdatedSettings.NumMaxPublicConnections.GetValue();
+		}
 
-			if (UpdatedSettings.SessionIdOverride.IsSet())
-			{
-				SessionSettings.SessionIdOverride = UpdatedSettings.SessionIdOverride.GetValue();
-			}
+		if (UpdatedSettings.NumOpenPrivateConnections.IsSet())
+		{
+			SessionSettings.NumOpenPrivateConnections = UpdatedSettings.NumOpenPrivateConnections.GetValue();
+		}
 
-			for (const FName& Key : UpdatedSettings.RemovedCustomSettings)
-			{
-				SessionSettings.CustomSettings.Remove(Key);
-			}
+		if (UpdatedSettings.NumOpenPublicConnections.IsSet())
+		{
+			SessionSettings.NumOpenPublicConnections = UpdatedSettings.NumOpenPublicConnections.GetValue();
+		}
 
-			SessionSettings.CustomSettings.Append(UpdatedSettings.UpdatedCustomSettings);
+		// TODO: We may need some additional logic for schema changes
+		if (UpdatedSettings.SchemaName.IsSet())
+		{
+			SessionSettings.SchemaName = UpdatedSettings.SchemaName.GetValue();
+		}
 
-			for (const FOnlineAccountIdHandle& Key : UpdatedSettings.RemovedRegisteredUsers)
-			{
-				SessionSettings.RegisteredUsers.Remove(Key);
-			}
+		if (UpdatedSettings.SessionIdOverride.IsSet())
+		{
+			SessionSettings.SessionIdOverride = UpdatedSettings.SessionIdOverride.GetValue();
+		}
 
-			SessionSettings.RegisteredUsers.Append(UpdatedSettings.UpdatedRegisteredUsers);
+		for (const FName& Key : UpdatedSettings.RemovedCustomSettings)
+		{
+			SessionSettings.CustomSettings.Remove(Key);
+		}
 
-			for (const FOnlineAccountIdHandle& Key : UpdatedSettings.RemovedSessionMembers)
-			{
-				SessionSettings.SessionMembers.Remove(Key);
-			}
+		SessionSettings.CustomSettings.Append(UpdatedSettings.UpdatedCustomSettings);
+
+		for (const FOnlineAccountIdHandle& Key : UpdatedSettings.RemovedRegisteredUsers)
+		{
+			SessionSettings.RegisteredUsers.Remove(Key);
+		}
+
+		SessionSettings.RegisteredUsers.Append(UpdatedSettings.UpdatedRegisteredUsers);
+
+		for (const FOnlineAccountIdHandle& Key : UpdatedSettings.RemovedSessionMembers)
+		{
+			SessionSettings.SessionMembers.Remove(Key);
+		}
 
-			for (FSessionMemberUpdatesMap::TConstIterator It(UpdatedSettings.UpdatedSessionMembers); It; ++It)
+		for (FSessionMemberUpdatesMap::TConstIterator It(UpdatedSettings.UpdatedSessionMembers); It; ++It)
+		{
+			if (FSessionMember* SessionMember = SessionSettings.SessionMembers.Find(It.Key()))
 			{
-				if (FSessionMember* SessionMember = SessionSettings.SessionMembers.Find(It.Key()))
-				{
-					const FSessionMemberUpdate& SessionMemberUpdate = It.Value();
+				const FSessionMemberUpdate& SessionMemberUpdate = It.Value();
 
-					SessionMember->MemberSettings.Append(SessionMemberUpdate.UpdatedMemberSettings);
+				SessionMember->MemberSettings.Append(SessionMemberUpdate.UpdatedMemberSettings);
 
-					for (const FName& Key : SessionMemberUpdate.RemovedMemberSettings)
-					{
-						SessionMember->MemberSettings.Remove(Key);
-					}
+				for (const FName& Key : SessionMemberUpdate.RemovedMemberSettings)
+				{
+					SessionMember->MemberSettings.Remove(Key);
 				}
 			}
+		}
 
-			// We set the result and fire the event
-			Op.SetResult(FUpdateSession::Result{ *FoundSession });
+		// We set the result and fire the event
+		Op.SetResult(FUpdateSession::Result{ FoundSession });
 
-			FSessionUpdated SessionUpdatedEvent{ *FoundSession, UpdatedSettings };
-			SessionEvents.OnSessionUpdated.Broadcast(SessionUpdatedEvent);
-		}
-		else
-		{
-			Op.SetError(Errors::NotFound());
-		}
+		FSessionUpdated SessionUpdatedEvent{ FoundSession, UpdatedSettings };
+		SessionEvents.OnSessionUpdated.Broadcast(SessionUpdatedEvent);
 	})
 	.Enqueue(GetSerialQueue());
 
@@ -395,77 +310,46 @@ TOnlineAsyncOpHandle<FUpdateSession> FSessionsNull::UpdateSession(FUpdateSession
 
 TOnlineAsyncOpHandle<FFindSessions> FSessionsNull::FindSessions(FFindSessions::Params&& Params)
 {
-	if (!CurrentSessionSearch.IsValid())
-	{
-		CurrentSessionSearch = MakeShared<FFindSessions::Result>();
-		CurrentSessionSearchHandle = GetOp<FFindSessions>(MoveTemp(Params));
-
-		FindLANSessions();
-
-		return CurrentSessionSearchHandle->GetHandle();
-	}
-	else
-	{
-		TOnlineAsyncOpRef<FFindSessions> Operation = GetOp<FFindSessions>(MoveTemp(Params));
-		Operation->SetError(Errors::AlreadyPending());
-		return Operation->GetHandle();
-	}
-}
-
-TOnlineAsyncOpHandle<FJoinSession> FSessionsNull::JoinSession(FJoinSession::Params&& Params)
-{
-	TOnlineAsyncOpRef<FJoinSession> Op = GetOp<FJoinSession>(MoveTemp(Params));
-	const FJoinSession::Params& OpParams = Op->GetParams();
+	TOnlineAsyncOpRef<FFindSessions> Op = GetOp<FFindSessions>(MoveTemp(Params));
+	const FFindSessions::Params& OpParams = Op->GetParams();
 
-	// First, check if a session with that name already exists
-	if (TSharedRef<FSessionNull>* FoundSession = SessionsByName.Find(OpParams.SessionName))
+	FOnlineError ParamsCheck = CheckFindSessionsParams(OpParams);
+	if (ParamsCheck != Errors::Success())
 	{
-		Op->SetError(Errors::InvalidParams()); // TODO: New error: Session with name %s already exists
+		Op->SetError(MoveTemp(ParamsCheck));
 		return Op->GetHandle();
 	}
 
-	const FSessionSettings& SessionSettings = OpParams.Session->SessionSettings;
-
-	TArray<FOnlineAccountIdHandle> LocalUserIds;
-	LocalUserIds.Reserve(OpParams.LocalUsers.Num());
-	OpParams.LocalUsers.GenerateKeyArray(LocalUserIds);
-
-	for (const FOnlineAccountIdHandle& LocalUserId : LocalUserIds)
+	Op->Then([this](TOnlineAsyncOp<FFindSessions>& Op) mutable
 	{
-		if (SessionSettings.SessionMembers.Contains(LocalUserId))
+		const FFindSessions::Params& OpParams = Op.GetParams();
+
+		FOnlineError StateCheck = CheckFindSessionsState(OpParams);
+		if (StateCheck != Errors::Success())
 		{
-			Op->SetError(Errors::AccessDenied()); // TODO: Bespoke error: User %s already in session
-			return Op->GetHandle();
+			Op.SetError(MoveTemp(StateCheck));
+			return;
 		}
-	}
 
-	if (!SessionSettings.bAllowNewMembers)
-	{
-		Op->SetError(Errors::AccessDenied());
-		return Op->GetHandle();
-	}
+		CurrentSessionSearch = MakeShared<FFindSessions::Result>();
+		CurrentSessionSearchHandle = Op.AsShared();
 
-	if (!SessionSettings.bAllowUnregisteredPlayers)
-	{
-		for (const FOnlineAccountIdHandle& LocalUserId : LocalUserIds)
-		{
-			if (!SessionSettings.RegisteredUsers.Contains(LocalUserId))
-			{
-				Op->SetError(Errors::AccessDenied());
-				return Op->GetHandle();
-			}
-		}
-	}
+		FindLANSessions();
+	})
+	.Enqueue(GetSerialQueue());
 
-	if (SessionSettings.NumMaxPublicConnections > 0 && SessionSettings.NumOpenPublicConnections < LocalUserIds.Num())
-	{
-		Op->SetError(Errors::AccessDenied());
-		return Op->GetHandle();
-	}
+	return Op->GetHandle();
+}
 
-	if (SessionSettings.NumMaxPrivateConnections > 0 && SessionSettings.NumOpenPrivateConnections < LocalUserIds.Num())
+TOnlineAsyncOpHandle<FJoinSession> FSessionsNull::JoinSession(FJoinSession::Params&& Params)
+{
+	TOnlineAsyncOpRef<FJoinSession> Op = GetOp<FJoinSession>(MoveTemp(Params));
+	const FJoinSession::Params& OpParams = Op->GetParams();
+
+	FOnlineError ParamsCheck = CheckJoinSessionParams(OpParams);
+	if (ParamsCheck != Errors::Success())
 	{
-		Op->SetError(Errors::AccessDenied());
+		Op->SetError(MoveTemp(ParamsCheck));
 		return Op->GetHandle();
 	}
 
@@ -473,7 +357,13 @@ TOnlineAsyncOpHandle<FJoinSession> FSessionsNull::JoinSession(FJoinSession::Para
 	Op->Then([this](TOnlineAsyncOp<FJoinSession>& Op) mutable
 	{
 		const FJoinSession::Params& OpParams = Op.GetParams();
-		const FSessionSettings& SessionSettings = OpParams.Session->SessionSettings;
+
+		FOnlineError StateCheck = CheckJoinSessionState(OpParams);
+		if (StateCheck != Errors::Success())
+		{
+			Op.SetError(MoveTemp(StateCheck));
+			return;
+		}
 
 		TSharedRef<FSessionNull> NewSessionNullRef = MakeShared<FSessionNull>(FSessionNull(*OpParams.Session));
 		FSessionSettings NewSessionNullSettings = NewSessionNullRef->SessionSettings;
@@ -481,6 +371,7 @@ TOnlineAsyncOpHandle<FJoinSession> FSessionsNull::JoinSession(FJoinSession::Para
 		FSessionJoined SessionJoinedEvent;
 		SessionJoinedEvent.Session = NewSessionNullRef;
 
+		// TODO: This will move to AddSessionMember
 		TArray<FOnlineAccountIdHandle> LocalUserIds;
 		LocalUserIds.Reserve(OpParams.LocalUsers.Num());
 		OpParams.LocalUsers.GenerateKeyArray(LocalUserIds);
@@ -489,18 +380,18 @@ TOnlineAsyncOpHandle<FJoinSession> FSessionsNull::JoinSession(FJoinSession::Para
 		{
 			bool bReserveSlot = true;
 
-			if (const FRegisteredUser* RegisteredUser = SessionSettings.RegisteredUsers.Find(LocalUserId))
+			if (const FRegisteredUser* RegisteredUser = NewSessionNullSettings.RegisteredUsers.Find(LocalUserId))
 			{
 				bReserveSlot = !RegisteredUser->bHasReservedSlot;
 			}
 
 			if (bReserveSlot)
 			{
-				if (SessionSettings.NumOpenPublicConnections > 0)
+				if (NewSessionNullSettings.NumOpenPublicConnections > 0)
 				{
 					NewSessionNullSettings.NumOpenPublicConnections--;
 				}
-				else if (SessionSettings.NumOpenPrivateConnections > 0)
+				else if (NewSessionNullSettings.NumOpenPrivateConnections > 0)
 				{
 					NewSessionNullSettings.NumOpenPrivateConnections--;
 				}
@@ -531,133 +422,31 @@ TOnlineAsyncOpHandle<FLeaveSession> FSessionsNull::LeaveSession(FLeaveSession::P
 	{
 		const FLeaveSession::Params& OpParams = Op.GetParams();
 
-		if (TSharedRef<FSessionNull>* FoundSession = SessionsByName.Find(OpParams.SessionName))
+		FOnlineError StateCheck = CheckLeaveSessionState(OpParams);
+		if (StateCheck != Errors::Success())
 		{
-			// Only the host should stop the session at this point
-			if (FoundSession->Get().OwnerUserId == OpParams.LocalUserId && FoundSession->Get().SessionSettings.JoinPolicy == ESessionJoinPolicy::Public)
-			{
-				StopLANSession();
-			}
-
-			SessionsById.Remove((*FoundSession)->SessionId);
-			SessionsByName.Remove(OpParams.SessionName);
-
-			Op.SetResult(FLeaveSession::Result{ });
-
-			FSessionLeft SessionLeftEvent;
-			SessionLeftEvent.LocalUserIds.Add(OpParams.LocalUserId);
-			SessionLeftEvent.LocalUserIds.Append(OpParams.AdditionalLocalUsers);
-			SessionEvents.OnSessionLeft.Broadcast(SessionLeftEvent);
-		}
-		else
-		{
-			Op.SetError(Errors::NotFound());
+			Op.SetError(MoveTemp(StateCheck));
+			return;
 		}
-	})
-	.Enqueue(GetSerialQueue());
-
-	return Op->GetHandle();
-}
-
-// TODO: The logic for player registration vs adding a new session member is likely to be updated
-TOnlineAsyncOpHandle<FRegisterPlayers> FSessionsNull::RegisterPlayers(FRegisterPlayers::Params&& Params)
-{
-	TOnlineAsyncOpRef<FRegisterPlayers> Op = GetOp<FRegisterPlayers>(MoveTemp(Params));
 
-	Op->Then([this](TOnlineAsyncOp<FRegisterPlayers>& Op) mutable
-	{
-		const FRegisterPlayers::Params& OpParams = Op.GetParams();
-
-		if (TSharedRef<FSessionNull>* FoundSession = SessionsByName.Find(OpParams.SessionName))
-		{
-			FSessionSettings& SessionSettings = FoundSession->Get().SessionSettings;
-
-			for (const FOnlineAccountIdHandle& User : OpParams.TargetUsers)
-			{
-				FRegisteredUser& RegisteredUser = SessionSettings.RegisteredUsers.FindOrAdd(User);
-
-				// RegisterPlayers will not revoke a slot reserved for a registered user
-				if (!RegisteredUser.bHasReservedSlot && OpParams.bReserveSlot)
-				{
-					bool bSlotReserved = false;
-					if (SessionSettings.NumOpenPublicConnections > 0)
-					{
-						SessionSettings.NumOpenPublicConnections--;
-						bSlotReserved = true;
-					}
-					else if (SessionSettings.NumOpenPrivateConnections > 0)
-					{
-						SessionSettings.NumOpenPrivateConnections--;
-						bSlotReserved = true;
-					}
-
-					RegisteredUser.bHasReservedSlot = bSlotReserved;
-				}
-
-				if (const FSessionMember* SessionMember = SessionSettings.SessionMembers.Find(User))
-				{
-					RegisteredUser.bIsInSession = true;
-				}
-			}
-
-			Op.SetResult(FRegisterPlayers::Result{});
-		}
-		else
+		TSharedRef<FSession>& FoundSession = SessionsByName.FindChecked(OpParams.SessionName);
+		
+		// Only the host should stop the session at this point
+		if (FoundSession->OwnerUserId == OpParams.LocalUserId && FoundSession->SessionSettings.JoinPolicy == ESessionJoinPolicy::Public)
 		{
-			Op.SetError(Errors::NotFound());
+			StopLANSession();
 		}
-	})
-	.Enqueue(GetSerialQueue());
-
-	return Op->GetHandle();
-}
 
-// TODO: The logic for player registration removal vs removing a session member is likely to be updated
-TOnlineAsyncOpHandle<FUnregisterPlayers> FSessionsNull::UnregisterPlayers(FUnregisterPlayers::Params&& Params)
-{
-	TOnlineAsyncOpRef<FUnregisterPlayers> Op = GetOp<FUnregisterPlayers>(MoveTemp(Params));
-
-	Op->Then([this](TOnlineAsyncOp<FUnregisterPlayers>& Op) mutable
-	{
-		const FUnregisterPlayers::Params& OpParams = Op.GetParams();
+		// TODO: Remove all local users from session
 
-		if (TSharedRef<FSessionNull>* FoundSession = SessionsByName.Find(OpParams.SessionName))
-		{
-			FSessionSettings& SessionSettings = FoundSession->Get().SessionSettings;
+		SessionsById.Remove(FoundSession->SessionId);
+		SessionsByName.Remove(OpParams.SessionName);
 
-			for (const FOnlineAccountIdHandle& User : OpParams.TargetUsers)
-			{
-				if (const FRegisteredUser* RegisteredUser = SessionSettings.RegisteredUsers.Find(User))
-				{
-					SessionSettings.RegisteredUsers.Remove(User);
-
-					bool bRemoveUserFromSession = RegisteredUser->bIsInSession && OpParams.bRemoveUnregisteredPlayers;
-
-					if (bRemoveUserFromSession)
-					{
-						SessionSettings.SessionMembers.Remove(User);
-					}
-
-					if (RegisteredUser->bHasReservedSlot || bRemoveUserFromSession)
-					{
-						if (SessionSettings.NumOpenPublicConnections < SessionSettings.NumMaxPublicConnections)
-						{
-							SessionSettings.NumOpenPublicConnections++;
-						}
-						else if (SessionSettings.NumOpenPrivateConnections < SessionSettings.NumMaxPrivateConnections)
-						{
-							SessionSettings.NumOpenPrivateConnections++;
-						}
-					}
-				}
-			}
+		Op.SetResult(FLeaveSession::Result{ });
 
-			Op.SetResult(FUnregisterPlayers::Result{});
-		}
-		else
-		{
-			Op.SetError(Errors::NotFound());
-		}
+		FSessionLeft SessionLeftEvent;
+		SessionLeftEvent.LocalUserIds.Append(OpParams.LocalUsers);
+		SessionEvents.OnSessionLeft.Broadcast(SessionLeftEvent);
 	})
 	.Enqueue(GetSerialQueue());
 
@@ -747,9 +536,9 @@ void FSessionsNull::StopLANSession()
 void FSessionsNull::OnValidQueryPacketReceived(uint8* PacketData, int32 PacketLength, uint64 ClientNonce)
 {
 	// Iterate through all registered sessions and respond for each one that can be joinable
-	for (TMap<FName, TSharedRef<FSessionNull>>::TConstIterator It(SessionsByName); It; ++It)
+	for (TMap<FName, TSharedRef<FSession>>::TConstIterator It(SessionsByName); It; ++It)
 	{
-		const TSharedRef<FSessionNull>& Session = It.Value();
+		const TSharedRef<FSession>& Session = It.Value();
 
 		if (Session->SessionSettings.JoinPolicy == ESessionJoinPolicy::Public)
 		{
@@ -758,7 +547,7 @@ void FSessionsNull::OnValidQueryPacketReceived(uint8* PacketData, int32 PacketLe
 			LANSessionManager->CreateHostResponsePacket(Packet, ClientNonce);
 
 			// Add all the session details
-			AppendSessionToPacket(Packet, Session);
+			AppendSessionToPacket(Packet, FSessionNull::Cast(*Session));
 
 			// Broadcast this response so the client can see us
 			if (!Packet.HasOverflow())
@@ -780,7 +569,7 @@ void FSessionsNull::OnValidResponsePacketReceived(uint8* PacketData, int32 Packe
 		TSharedRef<FSessionNull> Session = MakeShared<FSessionNull>();
 
 		FNboSerializeFromBufferNullSvc Packet(PacketData, PacketLength);
-		ReadSessionFromPacket(Packet, Session);
+		ReadSessionFromPacket(Packet, FSessionNull::Cast(*Session));
 
 		CurrentSessionSearch->FoundSessions.Add(Session);
 	}
@@ -822,27 +611,27 @@ void FSessionsNull::OnLANSearchTimeout()
 	}
 }
 
-void FSessionsNull::AppendSessionToPacket(FNboSerializeToBufferNullSvc& Packet, const TSharedRef<FSessionNull>& Session)
+void FSessionsNull::AppendSessionToPacket(FNboSerializeToBufferNullSvc& Packet, const FSessionNull& Session)
 {
 	// We won't save CurrentState in the packet as all advertised sessions will be Valid
-	Packet << Session->OwnerUserId
-		<< Session->SessionId
-		<< *Session->OwnerInternetAddr;
+	Packet << Session.OwnerUserId
+		<< Session.SessionId
+		<< *Session.OwnerInternetAddr;
 
 	// TODO: Write session settings to packet, after SchemaVariant work
 }
 
-void FSessionsNull::ReadSessionFromPacket(FNboSerializeFromBufferNullSvc& Packet, const TSharedRef<FSessionNull>& Session)
+void FSessionsNull::ReadSessionFromPacket(FNboSerializeFromBufferNullSvc& Packet, FSessionNull& Session)
 {
-	Packet >> Session->OwnerUserId
-		>> Session->SessionId
-		>> *Session->OwnerInternetAddr;
+	Packet >> Session.OwnerUserId
+		>> Session.SessionId
+		>> *Session.OwnerInternetAddr;
 
 	// We'll set the connect address for the remote session as a custom parameter, so it can be read in OnlineServices' GetResolvedConnectString
 	FCustomSessionSetting ConnectString;
-	ConnectString.Data.Set<FString>(Session->OwnerInternetAddr->ToString(true));
+	ConnectString.Data.Set<FString>(Session.OwnerInternetAddr->ToString(true));
 	ConnectString.Visibility = ECustomSessionSettingVisibility::ViaOnlineService;
-	Session->SessionSettings.CustomSettings.Add(CONNECT_STRING_TAG, ConnectString);
+	Session.SessionSettings.CustomSettings.Add(CONNECT_STRING_TAG, ConnectString);
 
 	// TODO:: Read session settings from packet, after SchemaVariant work
 }
diff --git a/Engine/Plugins/Online/Experimental/OnlineServicesNull/Source/Public/Online/SessionsNull.h b/Engine/Plugins/Online/Experimental/OnlineServicesNull/Source/Public/Online/SessionsNull.h
--- a/Engine/Plugins/Online/Experimental/OnlineServicesNull/Source/Public/Online/SessionsNull.h
+++ b/Engine/Plugins/Online/Experimental/OnlineServicesNull/Source/Public/Online/SessionsNull.h
@@ -9,35 +9,12 @@ namespace UE::Online {
 
 class FOnlineServicesNull;
 
-class FOnlineSessionIdString
-{
-public:
-	FString Data;
-	FOnlineSessionIdHandle Handle;
-};
-
-class FOnlineSessionIdRegistryNull : public IOnlineSessionIdRegistry
+class FOnlineSessionIdRegistryNull : public TOnlineSessionIdStringRegistry<EOnlineServices::Null>
 {
 public:
 	static FOnlineSessionIdRegistryNull& Get();
 
-	FOnlineSessionIdHandle Find(FString SessionId) const;
-	FOnlineSessionIdHandle FindOrAdd(FString SessionId);
-	FOnlineSessionIdHandle GetNext();
-
-	// Begin IOnlineSessionIdRegistry
-	virtual FString ToLogString(const FOnlineSessionIdHandle& Handle) const override;
-	virtual TArray<uint8> ToReplicationData(const FOnlineSessionIdHandle& Handle) const override;
-	virtual FOnlineSessionIdHandle FromReplicationData(const TArray<uint8>& ReplicationString) override;
-	// End IOnlineSessionIdRegistry
-
-	virtual ~FOnlineSessionIdRegistryNull() = default;
-
-private:
-	const FOnlineSessionIdString* GetInternal(const FOnlineSessionIdHandle& Handle) const;
-
-	TArray<FOnlineSessionIdString> Ids;
-	TMap<FString, FOnlineSessionIdString> StringToId;
+	FOnlineSessionIdHandle GetNextSessionId();
 };
 
 class FSessionNull : public FSession
@@ -46,6 +23,12 @@ public:
 	FSessionNull();
 	FSessionNull(const FSession& InSession);
 
+	static FSessionNull& Cast(FSession& InSession);
+	static const FSessionNull& Cast(const FSession& InSession);
+
+private:
+	void Initialize();
+
 public:
 	/** The IP address for the session owner */
 	TSharedPtr<FInternetAddr> OwnerInternetAddr;
@@ -62,16 +45,11 @@ public:
 	virtual void Tick(float DeltaSeconds) override;
 
 	// ISessions
-	virtual TOnlineResult<FGetAllSessions> GetAllSessions(FGetAllSessions::Params&& Params) override;
-	virtual TOnlineResult<FGetSessionByName> GetSessionByName(FGetSessionByName::Params&& Params) override;
-	virtual TOnlineResult<FGetSessionById> GetSessionById(FGetSessionById::Params&& Params) override;
 	virtual TOnlineAsyncOpHandle<FCreateSession> CreateSession(FCreateSession::Params&& Params) override;
 	virtual TOnlineAsyncOpHandle<FUpdateSession> UpdateSession(FUpdateSession::Params&& Params) override;
 	virtual TOnlineAsyncOpHandle<FFindSessions> FindSessions(FFindSessions::Params&& Params) override;
 	virtual TOnlineAsyncOpHandle<FJoinSession> JoinSession(FJoinSession::Params&& Params) override;
 	virtual TOnlineAsyncOpHandle<FLeaveSession> LeaveSession(FLeaveSession::Params&& Params) override;
-	virtual TOnlineAsyncOpHandle<FRegisterPlayers> RegisterPlayers(FRegisterPlayers::Params&& Params) override;
-	virtual TOnlineAsyncOpHandle<FUnregisterPlayers> UnregisterPlayers(FUnregisterPlayers::Params&& Params) override;
 
 private:
 	/** LAN Methods */
@@ -81,18 +59,10 @@ private:
 	void OnValidQueryPacketReceived(uint8* PacketData, int32 PacketLength, uint64 ClientNonce);
 	void OnValidResponsePacketReceived(uint8* PacketData, int32 PacketLength);
 	void OnLANSearchTimeout();
-	void AppendSessionToPacket(class FNboSerializeToBufferNullSvc& Packet, const TSharedRef<FSessionNull>& Session);
-	void ReadSessionFromPacket(class FNboSerializeFromBufferNullSvc& Packet, const TSharedRef<FSessionNull>& Session);
+	void AppendSessionToPacket(class FNboSerializeToBufferNullSvc& Packet, const FSessionNull& Session);
+	void ReadSessionFromPacket(class FNboSerializeFromBufferNullSvc& Packet, FSessionNull& Session);
 	
 private:
-	FOnlineServicesNull& Services;
-
-	TMap<FName, TSharedRef<FSessionNull>> SessionsByName;
-	TMap<FOnlineSessionIdHandle, TSharedRef<FSessionNull>> SessionsById;
-
-	TSharedPtr<FFindSessions::Result> CurrentSessionSearch;
-	TSharedPtr<TOnlineAsyncOp<FFindSessions>> CurrentSessionSearchHandle;
-
 	uint32 PublicSessionsHosted = 0;
 
 	TSharedPtr<FLANSession> LANSessionManager;
