OSSv2: Use LobbyVariant as StatValue, add corresponding conversions and checks for Stats implementions and integration tests.

diff --git a/Engine/Plugins/Online/Experimental/OnlineServices/Source/OnlineServicesCommon/Private/Online/StatsCommon.cpp b/Engine/Plugins/Online/Experimental/OnlineServices/Source/OnlineServicesCommon/Private/Online/StatsCommon.cpp
--- a/Engine/Plugins/Online/Experimental/OnlineServices/Source/OnlineServicesCommon/Private/Online/StatsCommon.cpp
+++ b/Engine/Plugins/Online/Experimental/OnlineServices/Source/OnlineServicesCommon/Private/Online/StatsCommon.cpp
@@ -66,7 +66,7 @@ void LexFromString(EStatLeaderboardUpdateMethod& OutValue, const TCHAR* InStr)
 	{
 		OutValue = EStatLeaderboardUpdateMethod::KeepBest;
 	}
-	else if (FCString::Stricmp(InStr, TEXT("Set")) == 0)
+	else if (FCString::Stricmp(InStr, TEXT("Force")) == 0)
 	{
 		OutValue = EStatLeaderboardUpdateMethod::Force;
 	}
@@ -183,14 +183,11 @@ void FStatsCommon::ReadStatDefinitionsFromConfig()
 			StatDefinition.UsageFlags |= (uint32)StatUsageFlag;
 		}
 
-		if (StatDefinition.UsageFlags & (uint32)EStatUsageFlags::Achievement)
+		FText StatModifyMethod;
+		GConfig->GetText(ConfigSection, *FString::Printf(TEXT("StatDef_%d_ModifyMethod"), StatIdx), StatModifyMethod, GEngineIni);
+		if (!StatModifyMethod.IsEmpty())
 		{
-			FText StatModifyMethod;
-			GConfig->GetText(ConfigSection, *FString::Printf(TEXT("StatDef_%d_ModifyMethod"), StatIdx), StatModifyMethod, GEngineIni);
-			if (!StatModifyMethod.IsEmpty())
-			{
-				LexFromString(StatDefinition.ModifyMethod, *StatModifyMethod.ToString());
-			}
+			LexFromString(StatDefinition.ModifyMethod, *StatModifyMethod.ToString());
 		}
 
 		if (StatDefinition.UsageFlags & (uint32)EStatUsageFlags::Leaderboard)
diff --git a/Engine/Plugins/Online/Experimental/OnlineServices/Source/OnlineServicesInterface/Public/Online/Stats.h b/Engine/Plugins/Online/Experimental/OnlineServices/Source/OnlineServicesInterface/Public/Online/Stats.h
--- a/Engine/Plugins/Online/Experimental/OnlineServices/Source/OnlineServicesInterface/Public/Online/Stats.h
+++ b/Engine/Plugins/Online/Experimental/OnlineServices/Source/OnlineServicesInterface/Public/Online/Stats.h
@@ -6,10 +6,12 @@
 #include "Online/OnlineAsyncOpHandle.h"
 #include "Online/OnlineMeta.h"
 
+// TODO: Include Online/SchemaVariant.h instead when it's ready
+#include "Online/Lobbies.h"
+
 namespace UE::Online {
 
-//using FStatValue = SchemaVariant;
-using FStatValue = TVariant<bool, int32, int64, float, double, FString>;
+using FStatValue = FLobbyVariant;
 
 struct FUserStats
 {
diff --git a/Engine/Plugins/Online/Experimental/OnlineServicesEOSGS/Source/Private/Online/StatsEOSGS.cpp b/Engine/Plugins/Online/Experimental/OnlineServicesEOSGS/Source/Private/Online/StatsEOSGS.cpp
--- a/Engine/Plugins/Online/Experimental/OnlineServicesEOSGS/Source/Private/Online/StatsEOSGS.cpp
+++ b/Engine/Plugins/Online/Experimental/OnlineServicesEOSGS/Source/Private/Online/StatsEOSGS.cpp
@@ -30,18 +30,6 @@ void FStatsEOSGS::Initialize()
 	check(StatsHandle);
 }
 
-namespace Private
-{
-
-int32 ConvertToIntegerFromStatValue(const FStatValue& StatValue)
-{
-	int32 Value = 0;
-	// TODO: Convert to proper value when SchemaVariant is ready
-	return Value;
-}
-
-}
-
 TOnlineAsyncOpHandle<FUpdateStats> FStatsEOSGS::UpdateStats(FUpdateStats::Params&& Params)
 {
 	TOnlineAsyncOpRef<FUpdateStats> Op = GetOp<FUpdateStats>(MoveTemp(Params));
@@ -72,7 +60,7 @@ TOnlineAsyncOpHandle<FUpdateStats> FStatsEOSGS::UpdateStats(FUpdateStats::Params
 				EOSStat.ApiVersion = EOS_STATS_INGESTDATA_API_LATEST;
 				static_assert(EOS_STATS_INGESTDATA_API_LATEST == 1, "EOS_Stats_IngestData updated, check new fields");
 
-				EOSStat.IngestAmount = Private::ConvertToIntegerFromStatValue(StatPair.Value);
+				EOSStat.IngestAmount = static_cast<int32>(StatPair.Value.GetInt64());
 				FCStringAnsi::Strncpy(EOSStatNames[Index].StatName, TCHAR_TO_UTF8(*StatPair.Key), STAT_NAME_MAX_LENGTH_EOS);
 				EOSStat.StatName = EOSStatNames[Index].StatName;
 
@@ -154,8 +142,7 @@ void ReadStatsFromEOSResult(EOS_HStats StatsHandle, const EOS_Stats_OnQueryStats
 		EOS_Stats_Stat* ReadStat = nullptr;
 		if (EOS_Stats_CopyStatByName(StatsHandle, &Options, &ReadStat) == EOS_EResult::EOS_Success)
 		{
-			// TODO: Add the stats here when SchemaVariant is ready
-			//OutStats.Add(StatName, FStatValue(ReadStat->Value));
+			OutStats.Add(StatName, FStatValue(static_cast<int64>(ReadStat->Value)));
 			EOS_Stats_Stat_Release(ReadStat);
 		}
 		else
diff --git a/Engine/Plugins/Online/Experimental/OnlineServicesNull/Source/Private/Online/StatsNull.cpp b/Engine/Plugins/Online/Experimental/OnlineServicesNull/Source/Private/Online/StatsNull.cpp
--- a/Engine/Plugins/Online/Experimental/OnlineServicesNull/Source/Private/Online/StatsNull.cpp
+++ b/Engine/Plugins/Online/Experimental/OnlineServicesNull/Source/Private/Online/StatsNull.cpp
@@ -3,6 +3,7 @@
 #include "Online/StatsNull.h"
 #include "Online/AuthNull.h"
 #include "Online/OnlineServicesNull.h"
+#include "Math/UnrealMathUtility.h"
 
 namespace UE::Online {
 
@@ -52,34 +53,49 @@ TOnlineAsyncOpHandle<FUpdateStats> FStatsNull::UpdateStats(FUpdateStats::Params&
 			{
 				if (FStatDefinition* StatDefinition = StatDefinitions.Find(UpdateUserStatPair.Key))
 				{
-					if (StatDefinition->UsageFlags & (uint32)EStatUsageFlags::Achievement)
+					if (FStatValue* StatValue = UserStats.Find(UpdateUserStatPair.Key))
 					{
-						if (FStatValue* StatValue = UserStats.Find(UpdateUserStatPair.Key))
+						switch (StatDefinition->ModifyMethod)
 						{
+						case EStatModifyMethod::Set:
 							*StatValue = UpdateUserStatPair.Value;
-
-							// TODO: Apply the value properly when SchemaVariant is ready
-							//switch (StatDefinition->ModifyMethod)
-							//{
-							//case EStatModifyMethod::Set:
-							//	*StatValue = UpdateUserStatPair.Value;
-							//	break;
-							//case EStatModifyMethod::Sum:
-							//	*StatValue += UpdateUserStatPair.Value;
-							//	break;
-							//case EStatModifyMethod::Largest:
-							//	*StatValue = Max(*StatValue, UpdateUserStatPair.Value);
-							//	break;
-							//case EStatModifyMethod::Smallest:
-							//	*StatValue = Min(*StatValue, UpdateUserStatPair.Value);
-							//	break;
-							//}
-						}
-						else
-						{
-							UserStats.Emplace(UpdateUserStatPair.Key, UpdateUserStatPair.Value);
+							break;
+						case EStatModifyMethod::Sum:
+							if (StatValue->VariantData.IsType<double>())
+							{
+								StatValue->Set(StatValue->GetDouble() + UpdateUserStatPair.Value.GetDouble());
+							}
+							else if (StatValue->VariantData.IsType<int64>())
+							{
+								StatValue->Set(StatValue->GetInt64() + UpdateUserStatPair.Value.GetInt64());
+							}
+							break;
+						case EStatModifyMethod::Largest:
+							if (StatValue->VariantData.IsType<double>())
+							{
+								StatValue->Set(FMath::Max(StatValue->GetDouble(), UpdateUserStatPair.Value.GetDouble()));
+							}
+							else if (StatValue->VariantData.IsType<int64>())
+							{
+								StatValue->Set(FMath::Max(StatValue->GetInt64(), UpdateUserStatPair.Value.GetInt64()));
+							}
+							break;
+						case EStatModifyMethod::Smallest:
+							if (StatValue->VariantData.IsType<double>())
+							{
+								StatValue->Set(FMath::Min(StatValue->GetDouble(), UpdateUserStatPair.Value.GetDouble()));
+							}
+							else if (StatValue->VariantData.IsType<int64>())
+							{
+								StatValue->Set(FMath::Min(StatValue->GetInt64(), UpdateUserStatPair.Value.GetInt64()));
+							}
+							break;
 						}
 					}
+					else
+					{
+						UserStats.Emplace(UpdateUserStatPair.Key, UpdateUserStatPair.Value);
+					}
 
 					if (StatDefinition->UsageFlags & (uint32)EStatUsageFlags::Leaderboard)
 					{
diff --git a/Engine/Plugins/Online/Experimental/OnlineServicesOSSAdapter/Source/Private/Online/StatsOSSAdapter.cpp b/Engine/Plugins/Online/Experimental/OnlineServicesOSSAdapter/Source/Private/Online/StatsOSSAdapter.cpp
--- a/Engine/Plugins/Online/Experimental/OnlineServicesOSSAdapter/Source/Private/Online/StatsOSSAdapter.cpp
+++ b/Engine/Plugins/Online/Experimental/OnlineServicesOSSAdapter/Source/Private/Online/StatsOSSAdapter.cpp
@@ -26,7 +26,26 @@ void FStatsOSSAdapter::ConvertStatValueV2ToStatUpdateV1(const FString& StatName,
 {
 	FOnlineStatValue OnlineStatValue;
 
-	// TODO: Add the FOnlineStatValue conversion when SchemaVariant is ready
+	if (StatValue.VariantData.IsType<double>())
+	{
+		OnlineStatValue.SetValue(StatValue.GetDouble());
+	}
+	else if (StatValue.VariantData.IsType<int64>())
+	{
+		OnlineStatValue.SetValue(StatValue.GetInt64());
+	}
+	else if (StatValue.VariantData.IsType<bool>())
+	{
+		OnlineStatValue.SetValue(StatValue.GetBoolean());
+	}
+	else if (StatValue.VariantData.IsType<FString>())
+	{
+		OnlineStatValue.SetValue(StatValue.GetString());
+	}
+	else
+	{
+		ensureMsgf(false, TEXT("Can't convert stat value!"));
+	}
 
 	FOnlineStatUpdate::EOnlineStatModificationType OnlineStatModificationType = FOnlineStatUpdate::EOnlineStatModificationType::Unknown;
 
@@ -50,7 +69,16 @@ void FStatsOSSAdapter::ConvertStatValueV2ToStatUpdateV1(const FString& StatName,
 
 void FStatsOSSAdapter::ConvertStatValueV1ToStatValueV2(const FOnlineStatValue& OnlineStatValue, FStatValue& OutStatValue)
 {
-	// TODO: Add the conversion when SchemaVariant is ready
+	switch (OnlineStatValue.GetType())
+	{
+	case EOnlineKeyValuePairDataType::String: OnlineStatValue.GetValue(OutStatValue.VariantData.Get<FString>()); break;
+	case EOnlineKeyValuePairDataType::Bool: OnlineStatValue.GetValue(OutStatValue.VariantData.Get<bool>()); break;
+	case EOnlineKeyValuePairDataType::Double: OnlineStatValue.GetValue(OutStatValue.VariantData.Get<double>()); break;
+	case EOnlineKeyValuePairDataType::Int64: OnlineStatValue.GetValue(OutStatValue.VariantData.Get<int64>()); break;
+	default: 
+		ensureMsgf(false, TEXT("Can't convert online stat value!"));
+		break;
+	}
 }
 
 FOnlineError FStatsOSSAdapter::ConvertUpdateUsersStatsV2toV1(const TArray<FUserStats>& UpdateUsersStats, TArray<FOnlineStatsUserUpdatedStats>& OutUpdateUserStatsV1)
