OSSv2 AuthNULL implementation

diff --git a/Engine/Plugins/Online/Experimental/OnlineServicesNull/Source/OnlineServicesNull.Build.cs b/Engine/Plugins/Online/Experimental/OnlineServicesNull/Source/OnlineServicesNull.Build.cs
--- a/Engine/Plugins/Online/Experimental/OnlineServicesNull/Source/OnlineServicesNull.Build.cs
+++ b/Engine/Plugins/Online/Experimental/OnlineServicesNull/Source/OnlineServicesNull.Build.cs
@@ -18,6 +18,12 @@ public class OnlineServicesNull : ModuleRules
 			}
 		);
 
+		PrivateDependencyModuleNames.AddRange(
+			new string[] { 
+				"ApplicationCore",
+			}
+		);
+
 		if (Target.bCompileAgainstEngine)
 		{
 			PublicDependencyModuleNames.AddRange(
diff --git a/Engine/Plugins/Online/Experimental/OnlineServicesNull/Source/Private/Online/AuthNull.cpp b/Engine/Plugins/Online/Experimental/OnlineServicesNull/Source/Private/Online/AuthNull.cpp
--- a/Engine/Plugins/Online/Experimental/OnlineServicesNull/Source/Private/Online/AuthNull.cpp
+++ b/Engine/Plugins/Online/Experimental/OnlineServicesNull/Source/Private/Online/AuthNull.cpp
@@ -2,53 +2,58 @@
 
 #include "Online/AuthNull.h"
 
+#include "Algo/ForEach.h"
+#include "Algo/Transform.h"
+#include "Misc/CoreDelegates.h"
+#include "Misc/OutputDeviceRedirector.h"
 #include "Online/OnlineServicesNull.h"
 #include "Online/OnlineServicesNullTypes.h"
 #include "Online/AuthErrors.h"
 #include "Online/OnlineErrorDefinitions.h"
+#include "SocketSubsystem.h"
+
+#include "GenericPlatform/GenericPlatformInputDeviceMapper.h"
 
 namespace UE::Online {
 // Copied from OSS Null
 
-#define NULL_OSS_STRING_BUFFER_LENGTH 256
-#define NULL_MAX_TOKEN_SIZE 4096
-
-
-FAuthNull::FAuthNull(FOnlineServicesNull& InServices)
-	: FAuthCommon(InServices)
+struct FAuthNullConfig
 {
-}
+	bool bAddUserNumToNullId = false;
+	bool bForceStableNullId = false;
+};
 
-void FAuthNull::Initialize()
-{
-	FAuthCommon::Initialize();
+namespace Meta {
 
-}
+BEGIN_ONLINE_STRUCT_META(FAuthNullConfig)
+	ONLINE_STRUCT_FIELD(FAuthNullConfig, bAddUserNumToNullId),
+	ONLINE_STRUCT_FIELD(FAuthNullConfig, bForceStableNullId)
+END_ONLINE_STRUCT_META()
 
-void FAuthNull::PreShutdown()
-{
-}
+/* Meta*/ }
+
+namespace {
 
-FString FAuthNull::GenerateRandomUserId(int32 LocalUserNum)
+FString GenerateRandomUserId(const FAuthNullConfig& Config, FPlatformUserId PlatformUserId)
 {
 	FString HostName;
-	/*if(ISocketSubsystem::Get(PLATFORM_SOCKETSUBSYSTEM))
+	if(ISocketSubsystem* SocketSubsystem = ISocketSubsystem::Get(PLATFORM_SOCKETSUBSYSTEM))
 	{
-		if (!ISocketSubsystem::Get(PLATFORM_SOCKETSUBSYSTEM)->GetHostName(HostName))
+		if (!SocketSubsystem->GetHostName(HostName))
 		{
 			// could not get hostname, use address
 			bool bCanBindAll;
-			TSharedPtr<class FInternetAddr> Addr = ISocketSubsystem::Get(PLATFORM_SOCKETSUBSYSTEM)->GetLocalHostAddr(*GLog, bCanBindAll);
+			TSharedPtr<class FInternetAddr> Addr = SocketSubsystem->GetLocalHostAddr(*GLog, bCanBindAll);
 			HostName = Addr->ToString(false);
 		}
-	}*/
+	}
 
-	bool bUseStableNullId = false;
+	bool bUseStableNullId = Config.bForceStableNullId;
 	FString UserSuffix;
 
-	if (true)
+	if (Config.bAddUserNumToNullId)
 	{
-		UserSuffix = FString::Printf(TEXT("-%d"), LocalUserNum);
+		UserSuffix = FString::Printf(TEXT("-%d"), PlatformUserId.GetInternalId());
 	}
 	 
 #if !(UE_BUILD_SHIPPING && WITH_EDITOR)
@@ -69,96 +74,102 @@ FString FAuthNull::GenerateRandomUserId(int32 LocalUserNum)
 	return FString::Printf(TEXT("OSSV2-%s-%s%s"), *HostName, *FGuid::NewGuid().ToString(), *UserSuffix);
 }
 
-TOnlineAsyncOpHandle<FAuthLogin> FAuthNull::Login(FAuthLogin::Params&& Params)
+TSharedRef<FAccountInfoNull> CreateAccountInfo(const FAuthNullConfig& Config, FPlatformUserId PlatformUserId)
 {
-	TOnlineAsyncOpRef<FAuthLogin> Op = GetOp<FAuthLogin>(MoveTemp(Params));
-	int32 LocalUserNum = FPlatformMisc::GetUserIndexForPlatformUser(Op->GetParams().PlatformUserId);
-
-	Op->Then([this, LocalUserNum](TOnlineAsyncOp<FAuthLogin>& InAsyncOp) mutable
-	{
-		TSharedRef<FAccountInfoNull> AccountInfo = MakeShared<FAccountInfoNull>();
-		FString DisplayId = GenerateRandomUserId(LocalUserNum);
-		AccountInfo->PlatformUserId = InAsyncOp.GetParams().PlatformUserId;
-		AccountInfo->DisplayName = DisplayId;
-		AccountInfo->UserId = FOnlineAccountIdRegistryNull::Get().Create(DisplayId, AccountInfo->PlatformUserId);
-		AccountInfo->LoginStatus = ELoginStatus::LoggedIn;
-		FAuthLogin::Result Result{ AccountInfo };
-		AccountInfos.Add(AccountInfo->UserId, AccountInfo);
-		InAsyncOp.SetResult(MoveTemp(Result));
-	})
-	.Enqueue(GetSerialQueue());
-
-	return Op->GetHandle();
+	const FString DisplayId = GenerateRandomUserId(Config, PlatformUserId);
+	return MakeShared<FAccountInfoNull>(FAccountInfoNull{
+		FOnlineAccountIdRegistryNull::Get().Create(DisplayId, PlatformUserId),
+		PlatformUserId,
+		ELoginStatus::LoggedIn,
+		{ { AccountAttributeData::DisplayName, DisplayId } }
+		});
 }
 
-TOnlineAsyncOpHandle<FAuthLogout> FAuthNull::Logout(FAuthLogout::Params&& Params)
+/* anonymous*/ }
+
+TSharedPtr<FAccountInfoNull> FAccountInfoRegistryNULL::Find(FPlatformUserId PlatformUserId) const
 {
-	// todo: login status change delegates
-	TOnlineAsyncOpRef<FAuthLogout> Op = GetOp<FAuthLogout>(MoveTemp(Params));
+	return StaticCastSharedPtr<FAccountInfoNull>(Super::Find(PlatformUserId));
+}
 
-	Op->Then([this](TOnlineAsyncOp<FAuthLogout>& InAsyncOp) mutable
-	{
-		auto Result = GetAccountByAccountId({InAsyncOp.GetParams().LocalUserId});
-		if(Result.IsOk())
-		{
-			AccountInfos.Remove(Result.GetOkValue().AccountInfo->UserId);
-			InAsyncOp.SetResult(FAuthLogout::Result());
-		}
-		else
-		{
-			InAsyncOp.SetError(Errors::Unknown());
-		}
-	})
-	.Enqueue(GetSerialQueue());
+TSharedPtr<FAccountInfoNull> FAccountInfoRegistryNULL::Find(FOnlineAccountIdHandle AccountIdHandle) const
+{
+	return StaticCastSharedPtr<FAccountInfoNull>(Super::Find(AccountIdHandle));
+}
 
-	return Op->GetHandle();
+void FAccountInfoRegistryNULL::Register(const TSharedRef<FAccountInfoNull>& AccountInfoNULL)
+{
+	DoRegister(AccountInfoNULL);
 }
 
-TOnlineResult<FAuthGetAccountByAccountId> FAuthNull::GetAccountByAccountId(FAuthGetAccountByAccountId::Params&& Params)
+void FAccountInfoRegistryNULL::Unregister(FOnlineAccountIdHandle AccountId)
 {
-	if (TSharedRef<FAccountInfoNull>* const FoundAccount = AccountInfos.Find(Params.LocalUserId))
+	if (TSharedPtr<FAccountInfoNull> AccountInfoNULL = Find(AccountId))
 	{
-		return TOnlineResult<FAuthGetAccountByAccountId>({*FoundAccount});
+		DoUnregister(AccountInfoNULL.ToSharedRef());
 	}
 	else
 	{
-		// TODO: proper error
-		return TOnlineResult<FAuthGetAccountByAccountId>(Errors::Unknown());
+		UE_LOG(LogOnlineServices, Warning, TEXT("[FAccountInfoRegistryNULL::Unregister] Failed to find account [%s]."), *ToLogString(AccountId));
 	}
 }
 
-bool FAuthNull::IsLoggedIn(const FOnlineAccountIdHandle& AccountId) const
+FAuthNull::FAuthNull(FOnlineServicesNull& InServices)
+	: FAuthCommon(InServices)
 {
-	// TODO:  More logic?
-	return AccountInfos.Contains(AccountId);
 }
 
-TResult<FOnlineAccountIdHandle, FOnlineError> FAuthNull::GetAccountIdByLocalUserNum(int32 LocalUserNum) const
+void FAuthNull::Initialize()
 {
-	FPlatformUserId UserId = FPlatformMisc::GetPlatformUserForUserIndex(LocalUserNum);
-	for (const TPair<FOnlineAccountIdHandle, TSharedRef<FAccountInfoNull>>& AccountPair : AccountInfos)
+	FAuthCommon::Initialize();
+	InitializeUsers();
+}
+
+void FAuthNull::PreShutdown()
+{
+	FAuthCommon::PreShutdown();
+	UninitializeUsers();
+}
+
+const FAccountInfoRegistry& FAuthNull::GetAccountInfoRegistry() const
+{
+	return AccountInfoRegistryNULL;
+}
+
+void FAuthNull::InitializeUsers()
+{
+	FAuthNullConfig AuthNullConfig;
+	LoadConfig(AuthNullConfig);
+
+	// There is no "login" for Null - all local users are initialized as "logged in".
+	TArray<FPlatformUserId> Users;
+	IPlatformInputDeviceMapper::Get().GetAllActiveUsers(Users);
+	Algo::ForEach(Users, [&](FPlatformUserId PlatformUserId)
 	{
-		if (AccountPair.Value->PlatformUserId == UserId)
-		{
-			TResult<FOnlineAccountIdHandle, FOnlineError> Result(AccountPair.Key);
-			return Result;
-		}
-	}
-	TResult<FOnlineAccountIdHandle, FOnlineError> Result(Errors::Unknown()); // TODO: error code
-	return Result;
+		AccountInfoRegistryNULL.Register(CreateAccountInfo(AuthNullConfig, PlatformUserId));
+	});
+
+	// Setup hook to add new users when they become available.
+	IPlatformInputDeviceMapper::Get().GetOnInputDeviceConnectionChange().AddRaw(this, &FAuthNull::OnInputDeviceConnectionChange);
 }
 
-TOnlineResult<FAuthGetAccountByPlatformUserId> FAuthNull::GetAccountByPlatformUserId(FAuthGetAccountByPlatformUserId::Params&& Params)
+void FAuthNull::UninitializeUsers()
 {
-	for (const TPair<FOnlineAccountIdHandle, TSharedRef<FAccountInfoNull>>& AccountPair : AccountInfos)
+	IPlatformInputDeviceMapper::Get().GetOnInputDeviceConnectionChange().RemoveAll(this);
+}
+
+void FAuthNull::OnInputDeviceConnectionChange(EInputDeviceConnectionState NewConnectionState, FPlatformUserId PlatformUserId, FInputDeviceId InputDeviceId)
+{
+	// If this is a new platform user then register an entry for them so they will be seen as "logged-in".
+	if (!AccountInfoRegistryNULL.Find(PlatformUserId))
 	{
-		if (AccountPair.Value->PlatformUserId == Params.PlatformUserId)
-		{
-			return TOnlineResult<FAuthGetAccountByPlatformUserId>({AccountPair.Value});
-		}
+		FAuthNullConfig AuthNullConfig;
+		LoadConfig(AuthNullConfig);
+
+		TSharedRef<FAccountInfoNull> AccountInfo = CreateAccountInfo(AuthNullConfig, PlatformUserId);
+		AccountInfoRegistryNULL.Register(AccountInfo);
+		OnAuthLoginStatusChangedEvent.Broadcast(FAuthLoginStatusChanged{ AccountInfo, ELoginStatus::LoggedIn });
 	}
-	TOnlineResult<FAuthGetAccountByPlatformUserId> Result(Errors::Unknown()); // TODO: error code
-	return Result;
 }
 
 // FOnlineAccountIdRegistryNull
@@ -180,7 +191,7 @@ FOnlineAccountIdHandle FOnlineAccountIdRegistryNull::Find(FString UserId) const
 
 FOnlineAccountIdHandle FOnlineAccountIdRegistryNull::Find(FPlatformUserId UserId) const
 {
-	const FOnlineAccountIdString* Entry = LocalUserMap.Find(UserId.GetInternalId());
+	const FOnlineAccountIdString* Entry = LocalUserMap.Find(UserId);
 	if (Entry)
 	{
 		return Entry->Handle;
@@ -188,9 +199,9 @@ FOnlineAccountIdHandle FOnlineAccountIdRegistryNull::Find(FPlatformUserId UserId
 	return FOnlineAccountIdHandle();
 }
 
-FOnlineAccountIdHandle FOnlineAccountIdRegistryNull::Find(int32 UserId) const
+FOnlineAccountIdHandle FOnlineAccountIdRegistryNull::Find(int32 UserIndex) const
 {
-	const FOnlineAccountIdString* Entry = LocalUserMap.Find(UserId);
+	const FOnlineAccountIdString* Entry = LocalUserMap.Find(FPlatformMisc::GetPlatformUserForUserIndex(UserIndex));
 	if (Entry)
 	{
 		return Entry->Handle;
@@ -207,28 +218,28 @@ const FOnlineAccountIdString* FOnlineAccountIdRegistryNull::GetInternal(const FO
 	return nullptr;
 }
 
-FOnlineAccountIdHandle FOnlineAccountIdRegistryNull::Create(FString UserId, FPlatformUserId LocalUserIndex/* = FPlatformUserId::PLATFORMUSERID_NONE*/)
+FOnlineAccountIdHandle FOnlineAccountIdRegistryNull::Create(FString UserId, FPlatformUserId PlatformUserId/* = PLATFORMUSERID_NONE*/)
 {
 	FOnlineAccountIdHandle ExistingHandle = Find(UserId);
 	if(ExistingHandle.IsValid())
 	{
+		UE_LOG(LogOnlineServices, Error, TEXT("[FOnlineAccountIdRegistryNull::Create] Found a duplicate ID for local user %d."), FPlatformMisc::GetUserIndexForPlatformUser(PlatformUserId));
 		return ExistingHandle;
 	}
 
+	if (!PlatformUserId.IsValid())
+	{
+		UE_LOG(LogOnlineServices, Warning, TEXT("[FOnlineAccountIdRegistryNull::Create] Unable to create id: PlatformUserId is invalid."));
+		return FOnlineAccountIdHandle();
+	}
+
 	FOnlineAccountIdString& Id = Ids.Emplace_GetRef();
 	Id.AccountIndex = Ids.Num();
 	Id.Data = UserId;
 	Id.Handle = FOnlineAccountIdHandle(EOnlineServices::Null, Id.AccountIndex);
 	
 	StringToId.Add(UserId, Id);
-	if(LocalUserIndex.IsValid())
-	{
-		if(LocalUserMap.Contains(LocalUserIndex.GetInternalId()))
-		{
-			UE_LOG(LogTemp, Error, TEXT("OssNull: Found a duplicate ID for local user %d"), LocalUserIndex.GetInternalId());
-		}
-		LocalUserMap.Add(LocalUserIndex.GetInternalId(), Id);
-	}
+	LocalUserMap.Add(PlatformUserId, Id);
 
 	return Id.Handle;
 }
@@ -251,11 +262,11 @@ TArray<uint8> FOnlineAccountIdRegistryNull::ToReplicationData(const FOnlineAccou
 		TArray<uint8> ReplicationData;
 		ReplicationData.SetNumUninitialized(Id->Data.Len());
 		StringToBytes(Id->Data, ReplicationData.GetData(), Id->Data.Len());
-		UE_LOG(LogTemp, VeryVerbose, TEXT("StringToBytes on %s returned %d len"), *Id->Data, ReplicationData.Num())
+		UE_LOG(LogOnlineServices, VeryVerbose, TEXT("[FOnlineAccountIdRegistryNull::ToReplicationData] StringToBytes on %s returned %d len"), *Id->Data, ReplicationData.Num())
 		return ReplicationData;
 	}
 
-	return TArray<uint8>();;
+	return TArray<uint8>();
 }
 
 FOnlineAccountIdHandle FOnlineAccountIdRegistryNull::FromReplicationData(const TArray<uint8>& ReplicationData)
diff --git a/Engine/Plugins/Online/Experimental/OnlineServicesNull/Source/Public/Online/AuthNull.h b/Engine/Plugins/Online/Experimental/OnlineServicesNull/Source/Public/Online/AuthNull.h
--- a/Engine/Plugins/Online/Experimental/OnlineServicesNull/Source/Public/Online/AuthNull.h
+++ b/Engine/Plugins/Online/Experimental/OnlineServicesNull/Source/Public/Online/AuthNull.h
@@ -7,6 +7,8 @@
 
 namespace UE::Online {
 
+class FOnlineServicesNull;
+
 class FOnlineAccountIdString
 {
 public:
@@ -22,9 +24,9 @@ public:
 
 	FOnlineAccountIdHandle Find(FString UserId) const;
 	FOnlineAccountIdHandle Find(FPlatformUserId UserId) const;
-	FOnlineAccountIdHandle Find(int32 UserId) const;
+	FOnlineAccountIdHandle Find(int32 UserIndex) const;
 
-	FOnlineAccountIdHandle Create(FString UserId, FPlatformUserId LocalUserIndex = PLATFORMUSERID_NONE);
+	FOnlineAccountIdHandle Create(FString UserId, FPlatformUserId PlatformUserId = PLATFORMUSERID_NONE);
 
 	// Begin IOnlineAccountIdRegistry
 	virtual FString ToLogString(const FOnlineAccountIdHandle& Handle) const override;
@@ -39,12 +41,32 @@ private:
 
 	// how much of these are actually necessary?
 	TArray<FOnlineAccountIdString> Ids;
-	TMap<FString, FOnlineAccountIdString> StringToId; 
-	TMap<int32, FOnlineAccountIdString> LocalUserMap;
+	TMap<FString, FOnlineAccountIdString> StringToId;
+	TMap<FPlatformUserId, FOnlineAccountIdString> LocalUserMap;
 
 };
 
-class FOnlineServicesNull;
+struct FAccountInfoNull final : public FAccountInfo
+{
+};
+
+
+// Auth NULL is implemented in a way similar to console platforms where there is not an explicit
+// login / logout from online services. On those platforms the user account is picked either before
+// the game has started or as a part of selecting an input device.
+class ONLINESERVICESNULL_API FAccountInfoRegistryNULL final : public FAccountInfoRegistry
+{
+public:
+	using Super = FAccountInfoRegistry;
+
+	virtual ~FAccountInfoRegistryNULL() = default;
+
+	TSharedPtr<FAccountInfoNull> Find(FPlatformUserId PlatformUserId) const;
+	TSharedPtr<FAccountInfoNull> Find(FOnlineAccountIdHandle AccountIdHandle) const;
+
+	void Register(const TSharedRef<FAccountInfoNull>&UserAuthData);
+	void Unregister(FOnlineAccountIdHandle AccountId);
+};
 
 class ONLINESERVICESNULL_API FAuthNull : public FAuthCommon
 {
@@ -54,22 +76,15 @@ public:
 	FAuthNull(FOnlineServicesNull& InOwningSubsystem);
 	virtual void Initialize() override;
 	virtual void PreShutdown() override;
-	virtual TOnlineAsyncOpHandle<FAuthLogin> Login(FAuthLogin::Params&& Params) override;
-	virtual TOnlineAsyncOpHandle<FAuthLogout> Logout(FAuthLogout::Params&& Params) override;
-	virtual TOnlineResult<FAuthGetAccountByPlatformUserId> GetAccountByPlatformUserId(FAuthGetAccountByPlatformUserId::Params&& Params) override;
-	virtual TOnlineResult<FAuthGetAccountByAccountId> GetAccountByAccountId(FAuthGetAccountByAccountId::Params&& Params) override;
-
-	bool IsLoggedIn(const FOnlineAccountIdHandle& AccountId) const;
 
 protected:
-	FString GenerateRandomUserId(int32 LocalUserNum);
-	TResult<FOnlineAccountIdHandle, FOnlineError> GetAccountIdByLocalUserNum(int32 LocalUserNum) const;
+	virtual const FAccountInfoRegistry& GetAccountInfoRegistry() const override;
 
-	class FAccountInfoNull : public FAccountInfo
-	{
-	};
+	void InitializeUsers();
+	void UninitializeUsers();
+	void OnInputDeviceConnectionChange(EInputDeviceConnectionState NewConnectionState, FPlatformUserId PlatformUserId, FInputDeviceId InputDeviceId);
 
-	TMap<FOnlineAccountIdHandle, TSharedRef<FAccountInfoNull>> AccountInfos;
+	FAccountInfoRegistryNULL AccountInfoRegistryNULL;
 };
 
 /* UE::Online */ }
