Fix clang build issue with FAccountInfoNull

diff --git a/Engine/Plugins/Online/Experimental/OnlineServicesNull/Source/Private/Online/AuthNull.cpp b/Engine/Plugins/Online/Experimental/OnlineServicesNull/Source/Private/Online/AuthNull.cpp
--- a/Engine/Plugins/Online/Experimental/OnlineServicesNull/Source/Private/Online/AuthNull.cpp
+++ b/Engine/Plugins/Online/Experimental/OnlineServicesNull/Source/Private/Online/AuthNull.cpp
@@ -77,12 +77,12 @@ FString GenerateRandomUserId(const FAuthNullConfig& Config, FPlatformUserId Plat
 TSharedRef<FAccountInfoNull> CreateAccountInfo(const FAuthNullConfig& Config, FPlatformUserId PlatformUserId)
 {
 	const FString DisplayId = GenerateRandomUserId(Config, PlatformUserId);
-	return MakeShared<FAccountInfoNull>(FAccountInfoNull{
+	return MakeShared<FAccountInfoNull>(FAccountInfoNull{ {
 		FOnlineAccountIdRegistryNull::Get().Create(DisplayId, PlatformUserId),
 		PlatformUserId,
 		ELoginStatus::LoggedIn,
 		{ { AccountAttributeData::DisplayName, DisplayId } }
-		});
+		} });
 }
 
 /* anonymous*/ }
