Addition of OnlineServicesCommonEngineUtils class for all Engine related functionality needed in OnlineServices.
Removed obsolete dependencies to OnlineSubsystemUtils.
Addition of FSessionsLAN to act as a base class for all sessions implementations with LAN functionality.
Refactored Null implementation to extend FSessionLAN type.
Fixed issue with SessionId replication.
Refactored NboSerializerNullSvc to be a namespace instead of a class, to fix ambiguity issues. Adapted usage code in SessionsNull and LobbiesNull.

diff --git a/Engine/Plugins/Online/Experimental/OnlineServices/OnlineServices.uplugin b/Engine/Plugins/Online/Experimental/OnlineServices/OnlineServices.uplugin
--- a/Engine/Plugins/Online/Experimental/OnlineServices/OnlineServices.uplugin
+++ b/Engine/Plugins/Online/Experimental/OnlineServices/OnlineServices.uplugin
@@ -20,6 +20,11 @@
 			"Name": "OnlineServicesCommon",
 			"Type": "Runtime",
 			"LoadingPhase": "PostConfigInit"
+		},
+		{
+			"Name": "OnlineServicesCommonEngineUtils",
+			"Type": "Runtime",
+			"LoadingPhase": "PostConfigInit"
 		}
 	],
 	"Plugins": [
diff --git a/Engine/Plugins/Online/Experimental/OnlineServices/Source/OnlineServicesCommon/Public/Online/SessionsCommon.h b/Engine/Plugins/Online/Experimental/OnlineServices/Source/OnlineServicesCommon/Public/Online/SessionsCommon.h
--- a/Engine/Plugins/Online/Experimental/OnlineServices/Source/OnlineServicesCommon/Public/Online/SessionsCommon.h
+++ b/Engine/Plugins/Online/Experimental/OnlineServices/Source/OnlineServicesCommon/Public/Online/SessionsCommon.h
@@ -29,10 +29,12 @@ public:
 
 	virtual inline TArray<uint8> ToReplicationData(const FOnlineSessionIdHandle& Handle) const override
 	{
-		const FTCHARToUTF8 IdValueUtf8(BasicRegistry.FindIdValue(Handle));
+		const FString IdValue = BasicRegistry.FindIdValue(Handle);
+		const FTCHARToUTF8 IdValueUtf8(IdValue);
 
 		TArray<uint8> ReplicationData;
 		ReplicationData.SetNumUninitialized(IdValueUtf8.Length());
+
 		FMemory::Memcpy(ReplicationData.GetData(), IdValueUtf8.Get(), IdValueUtf8.Length());
 
 		return ReplicationData;
@@ -40,11 +42,14 @@ public:
 
 	virtual inline FOnlineSessionIdHandle FromReplicationData(const TArray<uint8>& ReplicationData) override
 	{
-		FString Result = BytesToString(ReplicationData.GetData(), ReplicationData.Num());
-		if (Result.Len() > 0)
+		const FUTF8ToTCHAR IdValueTCHAR((char*)ReplicationData.GetData(), ReplicationData.Num());
+		const FString IdValue = FString(IdValueTCHAR.Length(), IdValueTCHAR.Get());
+
+		if (!IdValue.IsEmpty())
 		{
-			return BasicRegistry.FindOrAddHandle(Result);
+			return BasicRegistry.FindOrAddHandle(IdValue);
 		}
+
 		return FOnlineSessionIdHandle();
 	}
 	// End IOnlineSessionIdRegistry
diff --git a/Engine/Plugins/Online/Experimental/OnlineServices/Source/OnlineServicesCommonEngineUtils/OnlineServicesCommonEngineUtils.Build.cs b/Engine/Plugins/Online/Experimental/OnlineServices/Source/OnlineServicesCommonEngineUtils/OnlineServicesCommonEngineUtils.Build.cs
new file mode 100644
--- /dev/null
+++ b/Engine/Plugins/Online/Experimental/OnlineServices/Source/OnlineServicesCommonEngineUtils/OnlineServicesCommonEngineUtils.Build.cs
@@ -0,0 +1,35 @@
+// Copyright Epic Games, Inc. All Rights Reserved.
+
+using UnrealBuildTool;
+
+public class OnlineServicesCommonEngineUtils : ModuleRules
+{
+	public OnlineServicesCommonEngineUtils(ReadOnlyTargetRules Target) : base(Target)
+	{
+		PCHUsage = ModuleRules.PCHUsageMode.UseExplicitOrSharedPCHs;
+
+		bUseUnity = false; // set false to catch non-unity compile issues while prototyping
+
+		if (Target.bCompileAgainstEngine)
+		{
+			PublicDependencyModuleNames.AddRange(
+				new string[]
+				{
+					"Engine"
+				}
+			);
+		}
+
+		PrivateDependencyModuleNames.AddRange(
+			new string[] {
+				"Core",
+				"CoreOnline",
+				"CoreUObject",
+				"OnlineServicesInterface",
+				"OnlineServicesCommon",
+				"OnlineBase",
+				"Sockets"
+			}
+		);
+	}
+}
diff --git a/Engine/Plugins/Online/Experimental/OnlineServices/Source/OnlineServicesCommonEngineUtils/Private/Online/OnlineServicesCommonEngineUtils.cpp b/Engine/Plugins/Online/Experimental/OnlineServices/Source/OnlineServicesCommonEngineUtils/Private/Online/OnlineServicesCommonEngineUtils.cpp
new file mode 100644
--- /dev/null
+++ b/Engine/Plugins/Online/Experimental/OnlineServices/Source/OnlineServicesCommonEngineUtils/Private/Online/OnlineServicesCommonEngineUtils.cpp
@@ -0,0 +1,60 @@
+// Copyright Epic Games, Inc. All Rights Reserved.
+
+#include "Online/OnlineServicesCommonEngineUtils.h"
+
+#if WITH_ENGINE
+#include "Engine/GameEngine.h"
+#include "Engine/NetDriver.h"
+#endif
+
+namespace UE::Online {
+
+#if WITH_ENGINE
+UWorld* GetWorldForOnline(FName InstanceName)
+{
+	UWorld* World = NULL;
+#if WITH_EDITOR
+	if (InstanceName.ToString() != LexToString(UE::Online::EOnlineServices::Default) && InstanceName != NAME_None)
+	{
+		FWorldContext& WorldContext = GEngine->GetWorldContextFromHandleChecked(InstanceName);
+		check(WorldContext.WorldType == EWorldType::Game || WorldContext.WorldType == EWorldType::PIE);
+		World = WorldContext.World();
+	}
+	else
+#endif
+	{
+		UGameEngine* GameEngine = Cast<UGameEngine>(GEngine);
+		World = GameEngine ? GameEngine->GetGameWorld() : NULL;
+	}
+
+	return World;
+}
+#endif
+
+int32 GetPortFromNetDriver(FName InstanceName)
+{
+	int32 Port = 0;
+#if WITH_ENGINE
+	if (GEngine)
+	{
+		UWorld* World = GetWorldForOnline(InstanceName);
+		UNetDriver* NetDriver = World ? GEngine->FindNamedNetDriver(World, NAME_GameNetDriver) : NULL;
+		if (NetDriver && NetDriver->GetNetMode() < NM_Client)
+		{
+			FString AddressStr = NetDriver->LowLevelGetNetworkNumber();
+			int32 Colon = AddressStr.Find(":", ESearchCase::IgnoreCase, ESearchDir::FromEnd);
+			if (Colon != INDEX_NONE)
+			{
+				FString PortStr = AddressStr.Mid(Colon + 1);
+				if (!PortStr.IsEmpty())
+				{
+					Port = FCString::Atoi(*PortStr);
+				}
+			}
+		}
+	}
+#endif
+	return Port;
+}
+
+/* UE::Online */ }
\ No newline at end of file
diff --git a/Engine/Plugins/Online/Experimental/OnlineServices/Source/OnlineServicesCommonEngineUtils/Private/Online/SessionsLAN.cpp b/Engine/Plugins/Online/Experimental/OnlineServices/Source/OnlineServicesCommonEngineUtils/Private/Online/SessionsLAN.cpp
new file mode 100644
--- /dev/null
+++ b/Engine/Plugins/Online/Experimental/OnlineServices/Source/OnlineServicesCommonEngineUtils/Private/Online/SessionsLAN.cpp
@@ -0,0 +1,246 @@
+// Copyright Epic Games, Inc. All Rights Reserved.
+
+#include "Online/SessionsLAN.h"
+
+#include "Misc/Guid.h"
+#include "Online/OnlineServicesCommon.h"
+#include "Online/OnlineServicesCommonEngineUtils.h"
+#include "SocketSubsystem.h"
+#include "UObject/CoreNet.h"
+
+namespace UE::Online {
+
+/** FSessionLAN */
+
+FSessionLAN::FSessionLAN()
+{
+	Initialize();
+}
+
+FSessionLAN::FSessionLAN(const FSessionLAN& InSession)
+	: FSession(InSession)
+	, OwnerInternetAddr(InSession.OwnerInternetAddr)
+{
+
+}
+
+void FSessionLAN::Initialize()
+{
+	bool bCanBindAll;
+	OwnerInternetAddr = ISocketSubsystem::Get(PLATFORM_SOCKETSUBSYSTEM)->GetLocalHostAddr(*GLog, bCanBindAll);
+
+	// The below is a workaround for systems that set hostname to a distinct address from 127.0.0.1 on a loop back interface.
+	// See e.g. https://www.debian.org/doc/manuals/debian-reference/ch05.en.html#_the_hostname_resolution
+	// and http://serverfault.com/questions/363095/why-does-my-hostname-appear-with-the-address-127-0-1-1-rather-than-127-0-0-1-in
+	// Since we bind to 0.0.0.0, we won't answer on 127.0.1.1, so we need to advertise ourselves as 127.0.0.1 for any other loop back address we may have.
+
+	uint32 HostIp = 0;
+	OwnerInternetAddr->GetIp(HostIp); // Will return in host order
+
+	// If this address is on loop back interface, advertise it as 127.0.0.1
+	if ((HostIp & 0xff000000) == 0x7f000000)
+	{
+		OwnerInternetAddr->SetIp(0x7f000001); // 127.0.0.1
+	}
+
+	// Now set the port that was configured
+	OwnerInternetAddr->SetPort(GetPortFromNetDriver(NAME_None));
+
+	if (OwnerInternetAddr->GetPort() == 0)
+	{
+		OwnerInternetAddr->SetPort(7777); // Default port
+	}
+}
+
+FSessionLAN& FSessionLAN::Cast(FSession& InSession)
+{
+	check(InSession.SessionId.GetOnlineServicesType() == EOnlineServices::Null);
+
+	return *static_cast<FSessionLAN*>(&InSession);
+}
+
+const FSessionLAN& FSessionLAN::Cast(const FSession& InSession)
+{
+	check(InSession.SessionId.GetOnlineServicesType() == EOnlineServices::Null);
+
+	return *static_cast<const FSessionLAN*>(&InSession);
+}
+
+/** FSessionsLAN */
+
+FSessionsLAN::FSessionsLAN(FOnlineServicesCommon& InServices)
+	: Super(InServices)
+	, LANSessionManager(MakeShared<FLANSession>())
+{
+
+}
+
+void FSessionsLAN::Tick(float DeltaSeconds)
+{
+	if (LANSessionManager->GetBeaconState() > ELanBeaconState::NotUsingLanBeacon)
+	{
+		LANSessionManager->Tick(DeltaSeconds);
+	}
+
+	Super::Tick(DeltaSeconds);
+}
+
+bool FSessionsLAN::TryHostLANSession()
+{
+	// The LAN Beacon can only broadcast one session at a time
+	if (LANSessionManager->GetBeaconState() == ELanBeaconState::NotUsingLanBeacon)
+	{
+		FOnValidQueryPacketDelegate QueryPacketDelegate = FOnValidQueryPacketDelegate::CreateThreadSafeSP(this, &FSessionsLAN::OnValidQueryPacketReceived);
+
+		if (LANSessionManager->Host(QueryPacketDelegate))
+		{
+			UE_LOG(LogTemp, VeryVerbose, TEXT("[FSessionsLAN::TryHostLANSession] LAN Beacon hosting..."));
+
+			PublicSessionsHosted++;
+
+			return true;
+		}
+		else
+		{
+			UE_LOG(LogTemp, VeryVerbose, TEXT("[FSessionsLAN::TryHostLANSession] LAN Beacon failed to host!"));
+
+			LANSessionManager->StopLANSession();
+		}
+	}
+	else
+	{
+		UE_LOG(LogTemp, VeryVerbose, TEXT("[FSessionsLAN::TryHostLANSession] LAN Beacon already in use!"));
+	}
+
+	return false;
+}
+
+void FSessionsLAN::FindLANSessions()
+{
+	// Recreate the unique identifier for this client
+	GenerateNonce((uint8*)&LANSessionManager->LanNonce, 8);
+
+	// Bind delegates
+	FOnValidResponsePacketDelegate ResponseDelegate = FOnValidResponsePacketDelegate::CreateThreadSafeSP(this, &FSessionsLAN::OnValidResponsePacketReceived);
+	FOnSearchingTimeoutDelegate TimeoutDelegate = FOnSearchingTimeoutDelegate::CreateThreadSafeSP(this, &FSessionsLAN::OnLANSearchTimeout);
+
+	FNboSerializeToBuffer Packet(LAN_BEACON_MAX_PACKET_SIZE);
+	LANSessionManager->CreateClientQueryPacket(Packet, LANSessionManager->LanNonce);
+	if (LANSessionManager->Search(Packet, ResponseDelegate, TimeoutDelegate) == false)
+	{
+		UE_LOG(LogTemp, Warning, TEXT("[FSessionsLAN::FindLANSessions] Search failed!"));
+
+		if (LANSessionManager->GetBeaconState() == ELanBeaconState::Searching)
+		{
+			LANSessionManager->StopLANSession();
+		}
+
+		// Trigger the delegate as having failed
+		CurrentSessionSearchHandle->SetError(Errors::RequestFailure()); // TODO: May need a new error type
+
+		CurrentSessionSearch.Reset();
+		CurrentSessionSearchHandle.Reset();
+
+		// If we were hosting public sessions before the search, we'll return the beacon to that state
+		if (PublicSessionsHosted > 0)
+		{
+			TryHostLANSession();
+		}
+	}
+	else
+	{
+		UE_LOG(LogTemp, Warning, TEXT("[]FSessionsLAN::FindLANSessions] Searching...."));
+	}
+}
+
+void FSessionsLAN::StopLANSession()
+{
+	check(PublicSessionsHosted > 0);
+
+	PublicSessionsHosted--;
+
+	if (PublicSessionsHosted == 0)
+	{
+		LANSessionManager->StopLANSession();
+	}
+}
+
+void FSessionsLAN::OnValidQueryPacketReceived(uint8* PacketData, int32 PacketLength, uint64 ClientNonce)
+{
+	// Iterate through all registered sessions and respond for each one that can be joinable
+	for (TMap<FName, TSharedRef<FSession>>::TConstIterator It(SessionsByName); It; ++It)
+	{
+		const TSharedRef<FSession>& Session = It.Value();
+
+		if (Session->SessionSettings.JoinPolicy == ESessionJoinPolicy::Public)
+		{
+			FNboSerializeToBuffer Packet(LAN_BEACON_MAX_PACKET_SIZE);
+			// Create the basic header before appending additional information
+			LANSessionManager->CreateHostResponsePacket(Packet, ClientNonce);
+
+			// Add all the session details
+			AppendSessionToPacket(Packet, FSessionLAN::Cast(*Session));
+
+			// Broadcast this response so the client can see us
+			if (!Packet.HasOverflow())
+			{
+				LANSessionManager->BroadcastPacket(Packet, Packet.GetByteCount());
+			}
+			else
+			{
+				UE_LOG(LogTemp, Warning, TEXT("[FSessionsLAN::OnValidQueryPacketReceived] LAN broadcast packet overflow, cannot broadcast on LAN"));
+			}
+		}
+	}
+}
+
+void FSessionsLAN::OnValidResponsePacketReceived(uint8* PacketData, int32 PacketLength)
+{
+	if (CurrentSessionSearch.IsValid())
+	{
+		TSharedRef<FSessionLAN> Session = MakeShared<FSessionLAN>();
+
+		FNboSerializeFromBuffer Packet(PacketData, PacketLength);
+		ReadSessionFromPacket(Packet, FSessionLAN::Cast(*Session));
+
+		CurrentSessionSearch->FoundSessions.Add(Session);
+	}
+}
+
+void FSessionsLAN::OnLANSearchTimeout()
+{
+	if (LANSessionManager->GetBeaconState() == ELanBeaconState::Searching)
+	{
+		LANSessionManager->StopLANSession();
+	}
+
+	if (!CurrentSessionSearch.IsValid())
+	{
+		UE_LOG(LogTemp, Warning, TEXT("[FSessionsLAN::OnLANSearchTimeout] Session search invalid"));
+
+		CurrentSessionSearchHandle->SetError(Errors::InvalidState());
+	}
+
+	if (CurrentSessionSearch->FoundSessions.IsEmpty())
+	{
+		UE_LOG(LogTemp, Warning, TEXT("[FSessionsLAN::OnLANSearchTimeout] Session search found no results"));
+
+		CurrentSessionSearchHandle->SetError(Errors::NotFound());
+	}
+
+	UE_LOG(LogTemp, Warning, TEXT("[FSessionsLAN::OnLANSearchTimeout] %d sessions found!"), CurrentSessionSearch->FoundSessions.Num());
+
+	FFindSessions::Result Result = *CurrentSessionSearch;
+	CurrentSessionSearchHandle->SetResult(MoveTemp(Result));
+
+	CurrentSessionSearch.Reset();
+	CurrentSessionSearchHandle.Reset();
+
+	// If we were hosting public sessions before the search, we'll return the beacon to that state
+	if (PublicSessionsHosted > 0)
+	{
+		TryHostLANSession();
+	}
+}
+
+/* UE::Online */ }
\ No newline at end of file
diff --git a/Engine/Plugins/Online/Experimental/OnlineServices/Source/OnlineServicesCommonEngineUtils/Public/Online/OnlineServicesCommonEngineUtils.h b/Engine/Plugins/Online/Experimental/OnlineServices/Source/OnlineServicesCommonEngineUtils/Public/Online/OnlineServicesCommonEngineUtils.h
new file mode 100644
--- /dev/null
+++ b/Engine/Plugins/Online/Experimental/OnlineServices/Source/OnlineServicesCommonEngineUtils/Public/Online/OnlineServicesCommonEngineUtils.h
@@ -0,0 +1,12 @@
+// Copyright Epic Games, Inc. All Rights Reserved.
+
+#pragma once
+
+#include "CoreMinimal.h"
+#include "Online/CoreOnline.h"
+
+namespace UE::Online {
+
+ONLINESERVICESCOMMONENGINEUTILS_API int32 GetPortFromNetDriver(FName InstanceName);
+
+/* UE::Online */ }
\ No newline at end of file
diff --git a/Engine/Plugins/Online/Experimental/OnlineServices/Source/OnlineServicesCommonEngineUtils/Public/Online/SessionsLAN.h b/Engine/Plugins/Online/Experimental/OnlineServices/Source/OnlineServicesCommonEngineUtils/Public/Online/SessionsLAN.h
new file mode 100644
--- /dev/null
+++ b/Engine/Plugins/Online/Experimental/OnlineServices/Source/OnlineServicesCommonEngineUtils/Public/Online/SessionsLAN.h
@@ -0,0 +1,56 @@
+// Copyright Epic Games, Inc. All Rights Reserved.
+#pragma once
+
+#include "CoreMinimal.h"
+#include "Online/LANBeacon.h"
+#include "Online/NboSerializer.h"
+#include "Online/SessionsCommon.h"
+
+namespace UE::Online {
+
+class FOnlineServicesCommon;
+
+class ONLINESERVICESCOMMONENGINEUTILS_API FSessionLAN : public FSession
+{
+public:
+	FSessionLAN();
+	FSessionLAN(const FSessionLAN& InSession);
+
+	static FSessionLAN& Cast(FSession& InSession);
+	static const FSessionLAN& Cast(const FSession& InSession);
+
+private:
+	void Initialize();
+
+public:
+	/** The IP address for the session owner */
+	TSharedPtr<FInternetAddr> OwnerInternetAddr;
+};
+
+class ONLINESERVICESCOMMONENGINEUTILS_API FSessionsLAN : public FSessionsCommon
+{
+public:
+	using Super = FSessionsCommon;
+
+	FSessionsLAN(FOnlineServicesCommon& InServices);
+	virtual ~FSessionsLAN() = default;
+
+	virtual void Tick(float DeltaSeconds) override;
+
+protected:
+	bool TryHostLANSession();
+	void FindLANSessions();
+	void StopLANSession();
+	void OnValidQueryPacketReceived(uint8* PacketData, int32 PacketLength, uint64 ClientNonce);
+	void OnValidResponsePacketReceived(uint8* PacketData, int32 PacketLength);
+	void OnLANSearchTimeout();
+	virtual void AppendSessionToPacket(FNboSerializeToBuffer& Packet, const FSessionLAN& Session) = 0;
+	virtual void ReadSessionFromPacket(FNboSerializeFromBuffer& Packet, FSessionLAN& Session) = 0;
+
+protected:
+	uint32 PublicSessionsHosted = 0;
+
+	TSharedRef<FLANSession> LANSessionManager;
+};
+
+/* UE::Online */ }
\ No newline at end of file
diff --git a/Engine/Plugins/Online/Experimental/OnlineServicesNull/Source/OnlineServicesNull.Build.cs b/Engine/Plugins/Online/Experimental/OnlineServicesNull/Source/OnlineServicesNull.Build.cs
--- a/Engine/Plugins/Online/Experimental/OnlineServicesNull/Source/OnlineServicesNull.Build.cs
+++ b/Engine/Plugins/Online/Experimental/OnlineServicesNull/Source/OnlineServicesNull.Build.cs
@@ -21,18 +21,9 @@ public class OnlineServicesNull : ModuleRules
 		PrivateDependencyModuleNames.AddRange(
 			new string[] { 
 				"ApplicationCore",
-				"OnlineBase"
+				"OnlineBase",
+				"OnlineServicesCommonEngineUtils"
 			}
 		);
-
-		if (Target.bCompileAgainstEngine)
-		{
-			PublicDependencyModuleNames.AddRange(
-				new string[]
-				{
-					"OnlineSubsystemUtils"
-				}
-			);
-		}
 	}
 }
diff --git a/Engine/Plugins/Online/Experimental/OnlineServicesNull/Source/Private/Online/LobbiesNull.cpp b/Engine/Plugins/Online/Experimental/OnlineServicesNull/Source/Private/Online/LobbiesNull.cpp
--- a/Engine/Plugins/Online/Experimental/OnlineServicesNull/Source/Private/Online/LobbiesNull.cpp
+++ b/Engine/Plugins/Online/Experimental/OnlineServicesNull/Source/Private/Online/LobbiesNull.cpp
@@ -2,14 +2,11 @@
 
 #include "Online/LobbiesNull.h"
 #include "Online/OnlineServicesNull.h"
+#include "Online/OnlineServicesCommonEngineUtils.h"
 #include "Misc/Guid.h"
 #include "SocketSubsystem.h"
-#include "Online/NboSerializerNullSvc.h"
 #include "UObject/CoreNet.h"
 #include "Online/LobbyRegistryNull.h"
-#if WITH_ENGINE
-	#include "OnlineSubsystemUtils.h" // Needed for GetPortFromNetDriver
-#endif //WITH_ENGINE
 
 namespace UE::Online { 
 
@@ -32,10 +29,7 @@ FLobbyNull::FLobbyNull()
 	}
 
 	// Now set the port that was configured
-#if WITH_ENGINE
-	// jntodo: can we put anything here if we don't have engine code?
 	HostAddr->SetPort(GetPortFromNetDriver(NAME_None));
-#endif //WITH_ENGINE
 	
 	if (HostAddr->GetPort() == 0)
 	{
@@ -188,7 +182,7 @@ uint32 FLobbiesNull::FindLANSession()
 	FOnValidResponsePacketDelegate ResponseDelegate = FOnValidResponsePacketDelegate::CreateRaw(this, &FLobbiesNull::OnValidResponsePacketReceived);
 	FOnSearchingTimeoutDelegate TimeoutDelegate = FOnSearchingTimeoutDelegate::CreateRaw(this, &FLobbiesNull::OnLANSearchTimeout);
 
-	FNboSerializeToBufferNullSvc Packet(LAN_BEACON_MAX_PACKET_SIZE);
+	FNboSerializeToBuffer Packet(LAN_BEACON_MAX_PACKET_SIZE);
 	LANSessionManager.CreateClientQueryPacket(Packet, LANSessionManager.LanNonce);
 	if (LANSessionManager.Search(Packet, ResponseDelegate, TimeoutDelegate) == false)
 	{
@@ -250,14 +244,16 @@ uint32 FLobbiesNull::UpdateLANStatus()
 	return Result;
 }
 
-void FLobbiesNull::AppendLobbyToPacket(FNboSerializeToBufferNullSvc& Packet, const TSharedRef<FLobbyNull>& Lobby)
+void FLobbiesNull::AppendLobbyToPacket(FNboSerializeToBuffer& Packet, const TSharedRef<FLobbyNull>& Lobby)
 {
+	using namespace NboSerializerNullSvc;
+
 	/** Owner of the session */
 	UE_LOG(LogTemp, Warning, TEXT("AppendLobbyToPacket: Appending IP %s"), *Lobby->HostAddr->ToString(true));
 
-	Packet << Lobby->Data->OwnerAccountId
-		<< Lobby->Data->Attributes
-		<< TEXT("Temp:Owner") // todo: owner name
+	SerializeToBuffer(Packet, Lobby->Data->OwnerAccountId);
+	SerializeToBuffer(Packet, Lobby->Data->Attributes);
+	Packet	<< TEXT("Temp:Owner") // todo: owner name
 		<< Lobby->HostAddrIp
 		<< Lobby->HostAddrPort
 		<< 100 // todo: num private connections
@@ -276,7 +272,7 @@ void FLobbiesNull::OnValidQueryPacketReceived(uint8* PacketData, int32 PacketLen
 
 		// TODO: Joinability (ensure lobby is joinable
 
-		FNboSerializeToBufferNullSvc Packet(LAN_BEACON_MAX_PACKET_SIZE);
+		FNboSerializeToBuffer Packet(LAN_BEACON_MAX_PACKET_SIZE);
 		// Create the basic header before appending additional information
 		LANSessionManager.CreateHostResponsePacket(Packet, ClientNonce);
 
@@ -295,17 +291,19 @@ void FLobbiesNull::OnValidQueryPacketReceived(uint8* PacketData, int32 PacketLen
 	}
 }
 
-void FLobbiesNull::ReadLobbyFromPacket(FNboSerializeFromBufferNullSvc& Packet, const TSharedRef<FLobbyNull>& Session)
+void FLobbiesNull::ReadLobbyFromPacket(FNboSerializeFromBuffer& Packet, const TSharedRef<FLobbyNull>& Session)
 {
+	using namespace NboSerializerNullSvc;
+
 	UE_LOG(LogTemp, Verbose, TEXT("Reading session information from server"));
 
 	/** Owner of the session */
 	FOnlineAccountIdHandle OwningUserId;
 	FString OwningUserName;
 	uint32 NumOpenPrivateConnections, NumOpenPublicConnections;
-	Packet >> OwningUserId
-		>> Session->Data->Attributes
-		>> OwningUserName
+	SerializeFromBuffer(Packet, OwningUserId);
+	SerializeFromBuffer(Packet, Session->Data->Attributes);
+	Packet	>> OwningUserName
 		>> Session->HostAddrIp
 		>> Session->HostAddrPort
 		>> NumOpenPrivateConnections
@@ -342,7 +340,7 @@ void FLobbiesNull::OnValidResponsePacketReceived(uint8* PacketData, int32 Packet
 		//Lobby->PingInMs = static_cast<int32>((FPlatformTime::Seconds() - SessionSearchStartInSeconds) * 1000);
 
 		// Prepare to read data from the packet
-		FNboSerializeFromBufferNullSvc Packet(PacketData, PacketLength);
+		FNboSerializeFromBuffer Packet(PacketData, PacketLength);
 		
 		ReadLobbyFromPacket(Packet, Lobby);
 
diff --git a/Engine/Plugins/Online/Experimental/OnlineServicesNull/Source/Private/Online/SessionsNull.cpp b/Engine/Plugins/Online/Experimental/OnlineServicesNull/Source/Private/Online/SessionsNull.cpp
--- a/Engine/Plugins/Online/Experimental/OnlineServicesNull/Source/Private/Online/SessionsNull.cpp
+++ b/Engine/Plugins/Online/Experimental/OnlineServicesNull/Source/Private/Online/SessionsNull.cpp
@@ -7,9 +7,6 @@
 #include "Online/OnlineServicesNull.h"
 #include "SocketSubsystem.h"
 #include "UObject/CoreNet.h"
-#if WITH_ENGINE
-	#include "OnlineSubsystemUtils.h" // Needed for GetPortFromNetDriver
-#endif
 
 namespace UE::Online {
 
@@ -26,79 +23,17 @@ FOnlineSessionIdHandle FOnlineSessionIdRegistryNull::GetNextSessionId()
 	return BasicRegistry.FindOrAddHandle(FGuid().ToString());
 }
 
-/** FSessionNull */
-
-FSessionNull::FSessionNull()
-{
-	Initialize();
-}
-
-FSessionNull::FSessionNull(const FSession& InSession)
-	: FSession(InSession)
-{
-	const FSessionNull& InSessionNull = FSessionNull::Cast(InSession);
-
-	OwnerInternetAddr = InSessionNull.OwnerInternetAddr;
-}
-
-void FSessionNull::Initialize()
-{
-	bool bCanBindAll;
-	OwnerInternetAddr = ISocketSubsystem::Get(PLATFORM_SOCKETSUBSYSTEM)->GetLocalHostAddr(*GLog, bCanBindAll);
-
-	// The below is a workaround for systems that set hostname to a distinct address from 127.0.0.1 on a loop back interface.
-	// See e.g. https://www.debian.org/doc/manuals/debian-reference/ch05.en.html#_the_hostname_resolution
-	// and http://serverfault.com/questions/363095/why-does-my-hostname-appear-with-the-address-127-0-1-1-rather-than-127-0-0-1-in
-	// Since we bind to 0.0.0.0, we won't answer on 127.0.1.1, so we need to advertise ourselves as 127.0.0.1 for any other loop back address we may have.
-
-	uint32 HostIp = 0;
-	OwnerInternetAddr->GetIp(HostIp); // Will return in host order
-
-	// If this address is on loop back interface, advertise it as 127.0.0.1
-	if ((HostIp & 0xff000000) == 0x7f000000)
-	{
-		OwnerInternetAddr->SetIp(0x7f000001); // 127.0.0.1
-	}
-
-#if WITH_ENGINE
-	// Now set the port that was configured
-	OwnerInternetAddr->SetPort(GetPortFromNetDriver(NAME_None));
-#endif
-
-	if (OwnerInternetAddr->GetPort() == 0)
-	{
-		OwnerInternetAddr->SetPort(7777); // Default port
-	}
-}
-
-FSessionNull& FSessionNull::Cast(FSession& InSession)
-{
-	check(InSession.SessionId.GetOnlineServicesType() == EOnlineServices::Null);
-	
-	return *static_cast<FSessionNull*>(&InSession);
-}
-
-const FSessionNull& FSessionNull::Cast(const FSession& InSession)
-{
-	check(InSession.SessionId.GetOnlineServicesType() == EOnlineServices::Null);
-
-	return *static_cast<const FSessionNull*>(&InSession);
-}
-
 /** FSessionsNull */
 
 FSessionsNull::FSessionsNull(FOnlineServicesNull& InServices)
 	: Super(InServices)
 {
-	if (!LANSessionManager.IsValid())
-	{
-		LANSessionManager = MakeShared<FLANSession>();
-	}
+
 }
 
 void FSessionsNull::Tick(float DeltaSeconds)
 {
-	if (LANSessionManager.IsValid() && LANSessionManager->GetBeaconState() > ELanBeaconState::NotUsingLanBeacon)
+	if (LANSessionManager->GetBeaconState() > ELanBeaconState::NotUsingLanBeacon)
 	{
 		LANSessionManager->Tick(DeltaSeconds);
 	}
@@ -365,7 +300,7 @@ TOnlineAsyncOpHandle<FJoinSession> FSessionsNull::JoinSession(FJoinSession::Para
 			return;
 		}
 
-		TSharedRef<FSessionNull> NewSessionNullRef = MakeShared<FSessionNull>(*OpParams.Session);
+		TSharedRef<FSessionNull> NewSessionNullRef = MakeShared<FSessionNull>(FSessionNull::Cast(*OpParams.Session));
 		FSessionSettings NewSessionNullSettings = NewSessionNullRef->SessionSettings;
 
 		// TODO: This will move to AddSessionMember
@@ -454,179 +389,27 @@ TOnlineAsyncOpHandle<FLeaveSession> FSessionsNull::LeaveSession(FLeaveSession::P
 	return Op->GetHandle();
 }
 
-/** LAN Methods */
-
-bool FSessionsNull::TryHostLANSession()
-{
-	// The LAN Beacon can only broadcast one session at a time
-	if (LANSessionManager->GetBeaconState() == ELanBeaconState::NotUsingLanBeacon)
-	{
-		FOnValidQueryPacketDelegate QueryPacketDelegate = FOnValidQueryPacketDelegate::CreateRaw(this, &FSessionsNull::OnValidQueryPacketReceived);
-
-		if (LANSessionManager->Host(QueryPacketDelegate))
-		{
-			UE_LOG(LogTemp, VeryVerbose, TEXT("[FSessionsNull::TryHostLANSession] LAN Beacon hosting..."));
-
-			PublicSessionsHosted++;
-
-			return true;
-		}
-		else
-		{
-			UE_LOG(LogTemp, VeryVerbose, TEXT("[FSessionsNull::TryHostLANSession] LAN Beacon failed to host!"));
-
-			LANSessionManager->StopLANSession();
-		}
-	}
-	else
-	{
-		UE_LOG(LogTemp, VeryVerbose, TEXT("[FSessionsNull::TryHostLANSession] LAN Beacon already in use!"));
-	}
-
-	return false;
-}
-
-void FSessionsNull::FindLANSessions()
-{
-	// Recreate the unique identifier for this client
-	GenerateNonce((uint8*)&LANSessionManager->LanNonce, 8);
-
-	// Bind delegates
-	FOnValidResponsePacketDelegate ResponseDelegate = FOnValidResponsePacketDelegate::CreateRaw(this, &FSessionsNull::OnValidResponsePacketReceived);
-	FOnSearchingTimeoutDelegate TimeoutDelegate = FOnSearchingTimeoutDelegate::CreateRaw(this, &FSessionsNull::OnLANSearchTimeout);
-
-	FNboSerializeToBufferNullSvc Packet(LAN_BEACON_MAX_PACKET_SIZE);
-	LANSessionManager->CreateClientQueryPacket(Packet, LANSessionManager->LanNonce);
-	if (LANSessionManager->Search(Packet, ResponseDelegate, TimeoutDelegate) == false)
-	{
-		UE_LOG(LogTemp, Warning, TEXT("[FSessionsNull::FindLANSessions] Search failed!"));
-
-		if (LANSessionManager->GetBeaconState() == ELanBeaconState::Searching)
-		{
-			LANSessionManager->StopLANSession();
-		}
-
-		// Trigger the delegate as having failed
-		CurrentSessionSearchHandle->SetError(Errors::RequestFailure()); // TODO: May need a new error type
-
-		CurrentSessionSearch.Reset();
-		CurrentSessionSearchHandle.Reset();
-
-		// If we were hosting public sessions before the search, we'll return the beacon to that state
-		if (PublicSessionsHosted > 0)
-		{
-			TryHostLANSession();
-		}
-	}
-	else
-	{
-		UE_LOG(LogTemp, Warning, TEXT("[]FSessionsNull::FindLANSessions] Searching...."));
-	}
-}
-
-void FSessionsNull::StopLANSession()
-{
-	PublicSessionsHosted--;
-
-	if (PublicSessionsHosted <= 0)
-	{
-		LANSessionManager->StopLANSession();
-	}
-}
+/** FSessionsLAN */
 
-void FSessionsNull::OnValidQueryPacketReceived(uint8* PacketData, int32 PacketLength, uint64 ClientNonce)
+void FSessionsNull::AppendSessionToPacket(FNboSerializeToBuffer& Packet, const FSessionNull& Session)
 {
-	// Iterate through all registered sessions and respond for each one that can be joinable
-	for (TMap<FName, TSharedRef<FSession>>::TConstIterator It(SessionsByName); It; ++It)
-	{
-		const TSharedRef<FSession>& Session = It.Value();
-
-		if (Session->SessionSettings.JoinPolicy == ESessionJoinPolicy::Public)
-		{
-			FNboSerializeToBufferNullSvc Packet(LAN_BEACON_MAX_PACKET_SIZE);
-			// Create the basic header before appending additional information
-			LANSessionManager->CreateHostResponsePacket(Packet, ClientNonce);
-
-			// Add all the session details
-			AppendSessionToPacket(Packet, FSessionNull::Cast(*Session));
-
-			// Broadcast this response so the client can see us
-			if (!Packet.HasOverflow())
-			{
-				LANSessionManager->BroadcastPacket(Packet, Packet.GetByteCount());
-			}
-			else
-			{
-				UE_LOG(LogTemp, Warning, TEXT("[FSessionsNull::OnValidQueryPacketReceived] LAN broadcast packet overflow, cannot broadcast on LAN"));
-			}
-		}
-	}
-}
+	using namespace NboSerializerNullSvc;
 
-void FSessionsNull::OnValidResponsePacketReceived(uint8* PacketData, int32 PacketLength)
-{
-	if (CurrentSessionSearch.IsValid())
-	{
-		TSharedRef<FSessionNull> Session = MakeShared<FSessionNull>();
-
-		FNboSerializeFromBufferNullSvc Packet(PacketData, PacketLength);
-		ReadSessionFromPacket(Packet, FSessionNull::Cast(*Session));
-
-		CurrentSessionSearch->FoundSessions.Add(Session);
-	}
-}
-
-void FSessionsNull::OnLANSearchTimeout()
-{
-	if (LANSessionManager->GetBeaconState() == ELanBeaconState::Searching)
-	{
-		LANSessionManager->StopLANSession();
-	}
-
-	if (!CurrentSessionSearch.IsValid())
-	{
-		UE_LOG(LogTemp, Warning, TEXT("[FSessionsNull::OnLANSearchTimeout] Session search invalid"));
-
-		CurrentSessionSearchHandle->SetError(Errors::InvalidState());
-	}
-
-	if (CurrentSessionSearch->FoundSessions.IsEmpty())
-	{
-		UE_LOG(LogTemp, Warning, TEXT("[FSessionsNull::OnLANSearchTimeout] Session search found no results"));
-
-		CurrentSessionSearchHandle->SetError(Errors::NotFound());
-	}
-
-	UE_LOG(LogTemp, Warning, TEXT("[FSessionsNull::OnLANSearchTimeout] %d sessions found!"), CurrentSessionSearch->FoundSessions.Num());
-
-	FFindSessions::Result Result = *CurrentSessionSearch;
-	CurrentSessionSearchHandle->SetResult(MoveTemp(Result));
-
-	CurrentSessionSearch.Reset();
-	CurrentSessionSearchHandle.Reset();
-
-	// If we were hosting public sessions before the search, we'll return the beacon to that state
-	if (PublicSessionsHosted > 0)
-	{
-		TryHostLANSession();
-	}
-}
-
-void FSessionsNull::AppendSessionToPacket(FNboSerializeToBufferNullSvc& Packet, const FSessionNull& Session)
-{
 	// We won't save CurrentState in the packet as all advertised sessions will be Valid
-	Packet << Session.OwnerUserId
-		<< Session.SessionId
-		<< *Session.OwnerInternetAddr;
+	SerializeToBuffer(Packet, Session.OwnerUserId);
+	SerializeToBuffer(Packet, Session.SessionId);
+	Packet << *Session.OwnerInternetAddr;
 
 	// TODO: Write session settings to packet, after SchemaVariant work
 }
 
-void FSessionsNull::ReadSessionFromPacket(FNboSerializeFromBufferNullSvc& Packet, FSessionNull& Session)
+void FSessionsNull::ReadSessionFromPacket(FNboSerializeFromBuffer& Packet, FSessionNull& Session)
 {
-	Packet >> Session.OwnerUserId
-		>> Session.SessionId
-		>> *Session.OwnerInternetAddr;
+	using namespace NboSerializerNullSvc;
+
+	SerializeFromBuffer(Packet, Session.OwnerUserId);
+	SerializeFromBuffer(Packet, Session.SessionId);
+	Packet >> *Session.OwnerInternetAddr;
 
 	// We'll set the connect address for the remote session as a custom parameter, so it can be read in OnlineServices' GetResolvedConnectString
 	FCustomSessionSetting ConnectString;
diff --git a/Engine/Plugins/Online/Experimental/OnlineServicesNull/Source/Public/Online/LobbiesNull.h b/Engine/Plugins/Online/Experimental/OnlineServicesNull/Source/Public/Online/LobbiesNull.h
--- a/Engine/Plugins/Online/Experimental/OnlineServicesNull/Source/Public/Online/LobbiesNull.h
+++ b/Engine/Plugins/Online/Experimental/OnlineServicesNull/Source/Public/Online/LobbiesNull.h
@@ -5,6 +5,7 @@
 #include "Online/LobbiesCommon.h"
 #include "Online/Lobbies.h"
 #include "Online/LANBeacon.h"
+#include "Online/NboSerializerNullSvc.h"
 
 namespace UE::Online {
 
@@ -47,9 +48,9 @@ private:
 	uint32 JoinLANSession(int32 PlayerNum, TSharedRef<FLobbyNull>& Lobby);
 	uint32 FindLANSession();
 	uint32 FinalizeLANSearch();
-	void AppendLobbyToPacket(class FNboSerializeToBufferNullSvc& Packet, const TSharedRef<FLobbyNull>& Lobby);
+	void AppendLobbyToPacket(FNboSerializeToBuffer& Packet, const TSharedRef<FLobbyNull>& Lobby);
 	//void AppendSessionSettingsToPacket(class FNboSerializeToBufferNull& Packet, FOnlineSessionSettings* SessionSettings);
-	void ReadLobbyFromPacket(class FNboSerializeFromBufferNullSvc& Packet, const TSharedRef<FLobbyNull>& Session);
+	void ReadLobbyFromPacket(FNboSerializeFromBuffer& Packet, const TSharedRef<FLobbyNull>& Session);
 	//void ReadSettingsFromPacket(class FNboSerializeFromBufferNull& Packet, FOnlineSessionSettings& SessionSettings);
 	void OnValidQueryPacketReceived(uint8* PacketData, int32 PacketLength, uint64 ClientNonce);
 	void OnValidResponsePacketReceived(uint8* PacketData, int32 PacketLength);
diff --git a/Engine/Plugins/Online/Experimental/OnlineServicesNull/Source/Public/Online/NboSerializerNullSvc.h b/Engine/Plugins/Online/Experimental/OnlineServicesNull/Source/Public/Online/NboSerializerNullSvc.h
--- a/Engine/Plugins/Online/Experimental/OnlineServicesNull/Source/Public/Online/NboSerializerNullSvc.h
+++ b/Engine/Plugins/Online/Experimental/OnlineServicesNull/Source/Public/Online/NboSerializerNullSvc.h
@@ -14,102 +14,70 @@
  */
 namespace UE::Online {
 
-class FNboSerializeToBufferNullSvc : public FNboSerializeToBuffer
+namespace NboSerializerNullSvc {
+
+/** NboSerializeToBuffer methods */
+
+inline void SerializeToBuffer(FNboSerializeToBuffer& Ar, const FOnlineAccountIdHandle& UniqueId)
 {
-public:
-	/** Default constructor zeros num bytes*/
-	FNboSerializeToBufferNullSvc() :
-		FNboSerializeToBuffer(512)
-	{
-	}
+	TArray<uint8> Data = FOnlineAccountIdRegistryNull::Get().ToReplicationData(UniqueId);
+	Ar << Data.Num();
+	Ar.WriteBinary(Data.GetData(), Data.Num());
+}
 
-	/** Constructor specifying the size to use */
-	FNboSerializeToBufferNullSvc(uint32 Size) :
-		FNboSerializeToBuffer(Size)
-	{
-	}
+inline void SerializeToBuffer(FNboSerializeToBuffer& Ar, const FOnlineSessionIdHandle& SessionId)
+{
+	TArray<uint8> Data = FOnlineSessionIdRegistryNull::Get().ToReplicationData(SessionId);
+	Ar << Data.Num();
+	Ar.WriteBinary(Data.GetData(), Data.Num());
+}
 
-	/**
-	 * Adds Null Id to the buffer
-	 */
-	friend inline FNboSerializeToBufferNullSvc& operator<<(FNboSerializeToBufferNullSvc& Ar, const FOnlineAccountIdHandle& UniqueId)
+inline void SerializeToBuffer(FNboSerializeToBuffer& Ar, const TMap<FLobbyAttributeId, FLobbyVariant>& Map)
+{
+	Ar << Map.Num();
+	for (const TPair<FLobbyAttributeId, FLobbyVariant>& Pair : Map)
 	{
-		TArray<uint8> Data = FOnlineAccountIdRegistryNull::Get().ToReplicationData(UniqueId);
-		Ar << Data.Num();
-		Ar.WriteBinary(Data.GetData(), Data.Num());
-		return Ar;
+		Ar << Pair.Key;
+		Ar << Pair.Value.GetString();
 	}
+}
 
-	friend inline FNboSerializeToBufferNullSvc& operator<<(FNboSerializeToBufferNullSvc& Ar, const FOnlineSessionIdHandle& SessionId)
-	{
-		TArray<uint8> Data = FOnlineSessionIdRegistryNull::Get().ToReplicationData(SessionId);
-		Ar << Data.Num();
-		Ar.WriteBinary(Data.GetData(), Data.Num());
-		return Ar;
-	}
+/** NboSerializeFromBuffer methods */
 
-	friend inline FNboSerializeToBufferNullSvc& operator<<(FNboSerializeToBufferNullSvc& Ar, const TMap<FLobbyAttributeId, FLobbyVariant>& Map)
-	{
-		Ar << Map.Num();
-		for (const TPair<FLobbyAttributeId, FLobbyVariant>& Pair : Map)
-		{
-			Ar << Pair.Key;
-			Ar << Pair.Value.GetString();
-		}
-		return Ar;
-	}
-};
+inline void SerializeFromBuffer(FNboSerializeFromBuffer& Ar, FOnlineAccountIdHandle& UniqueId)
+{
+	TArray<uint8> Data;
+	int32 Size;
+	Ar >> Size;
+	Ar.ReadBinaryArray(Data, Size);
 
-/**
- * Class used to write data into packets for sending via system link
- */
-class FNboSerializeFromBufferNullSvc : public FNboSerializeFromBuffer
+	UniqueId = FOnlineAccountIdRegistryNull::Get().FromReplicationData(Data);
+}
+
+inline void SerializeFromBuffer(FNboSerializeFromBuffer& Ar, FOnlineSessionIdHandle& SessionId)
 {
-public:
-	/**
-	 * Initializes the buffer, size, and zeros the read offset
-	 */
-	FNboSerializeFromBufferNullSvc(uint8* Packet,int32 Length) :
-		FNboSerializeFromBuffer(Packet,Length)
-	{
-	}
-	/**
-	 * Reads Null Id from the buffer
-	 */
-	friend inline FNboSerializeFromBufferNullSvc& operator>>(FNboSerializeFromBufferNullSvc& Ar, FOnlineAccountIdHandle& UniqueId)
-	{
-		TArray<uint8> Data;
-		int32 Size;
-		Ar >> Size;
-		Ar.ReadBinaryArray(Data, Size);
-		UniqueId = FOnlineAccountIdRegistryNull::Get().FromReplicationData(Data);
-		return Ar;
-	}
+	TArray<uint8> Data;
+	int32 Size;
+	Ar >> Size;
+	Ar.ReadBinaryArray(Data, Size);
 
-	friend inline FNboSerializeFromBufferNullSvc& operator>>(FNboSerializeFromBufferNullSvc& Ar, FOnlineSessionIdHandle& SessionId)
-	{
-		TArray<uint8> Data;
-		int32 Size;
-		Ar >> Size;
-		Ar.ReadBinaryArray(Data, Size);
-		SessionId = FOnlineSessionIdRegistryNull::Get().FromReplicationData(Data);
-		return Ar;
-	}
-	
-	friend inline FNboSerializeFromBufferNullSvc& operator>>(FNboSerializeFromBufferNullSvc& Ar, TMap<FLobbyAttributeId, FLobbyVariant>& Map)
+	SessionId = FOnlineSessionIdRegistryNull::Get().FromReplicationData(Data);
+}
+
+inline void SerializeFromBuffer(FNboSerializeFromBuffer& Ar, TMap<FLobbyAttributeId, FLobbyVariant>& Map)
+{
+	int32 Size;
+	Ar >> Size;
+	for (int i = 0; i < Size; i++)
 	{
-		int32 Size;
-		Ar >> Size;
-		for(int i = 0; i < Size; i++)
-		{
-			FLobbyAttributeId LobbyAttributeId;
-			FString LobbyData;
-			Ar >> LobbyAttributeId;
-			Ar >> LobbyData;
-			Map.Emplace(MoveTemp(LobbyAttributeId), MoveTemp(LobbyData));
-		}
-		return Ar;
+		FLobbyAttributeId LobbyAttributeId;
+		FString LobbyData;
+		Ar >> LobbyAttributeId;
+		Ar >> LobbyData;
+		Map.Emplace(MoveTemp(LobbyAttributeId), MoveTemp(LobbyData));
 	}
-};
+}
+
+/* NboSerializerNullSvc */ }
 
-} //namespace UE::Online
\ No newline at end of file
+/* UE::Online */ }
\ No newline at end of file
diff --git a/Engine/Plugins/Online/Experimental/OnlineServicesNull/Source/Public/Online/SessionsNull.h b/Engine/Plugins/Online/Experimental/OnlineServicesNull/Source/Public/Online/SessionsNull.h
--- a/Engine/Plugins/Online/Experimental/OnlineServicesNull/Source/Public/Online/SessionsNull.h
+++ b/Engine/Plugins/Online/Experimental/OnlineServicesNull/Source/Public/Online/SessionsNull.h
@@ -3,7 +3,7 @@
 
 #include "CoreMinimal.h"
 #include "Online/LANBeacon.h"
-#include "Online/SessionsCommon.h"
+#include "Online/SessionsLAN.h"
 
 namespace UE::Online {
 
@@ -17,27 +17,12 @@ public:
 	FOnlineSessionIdHandle GetNextSessionId();
 };
 
-class FSessionNull : public FSession
-{
-public:
-	FSessionNull();
-	FSessionNull(const FSession& InSession);
-
-	static FSessionNull& Cast(FSession& InSession);
-	static const FSessionNull& Cast(const FSession& InSession);
-
-private:
-	void Initialize();
+typedef FSessionLAN FSessionNull;
 
-public:
-	/** The IP address for the session owner */
-	TSharedPtr<FInternetAddr> OwnerInternetAddr;
-};
-
-class FSessionsNull : public FSessionsCommon
+class FSessionsNull : public FSessionsLAN
 {
 public:
-	using Super = FSessionsCommon;
+	using Super = FSessionsLAN;
 
 	FSessionsNull(FOnlineServicesNull& InServices);
 	virtual ~FSessionsNull() = default;
@@ -52,20 +37,9 @@ public:
 	virtual TOnlineAsyncOpHandle<FLeaveSession> LeaveSession(FLeaveSession::Params&& Params) override;
 
 private:
-	/** LAN Methods */
-	bool TryHostLANSession();
-	void FindLANSessions();
-	void StopLANSession();
-	void OnValidQueryPacketReceived(uint8* PacketData, int32 PacketLength, uint64 ClientNonce);
-	void OnValidResponsePacketReceived(uint8* PacketData, int32 PacketLength);
-	void OnLANSearchTimeout();
-	void AppendSessionToPacket(class FNboSerializeToBufferNullSvc& Packet, const FSessionNull& Session);
-	void ReadSessionFromPacket(class FNboSerializeFromBufferNullSvc& Packet, FSessionNull& Session);
-	
-private:
-	uint32 PublicSessionsHosted = 0;
-
-	TSharedPtr<FLANSession> LANSessionManager;
+	// FSessionsLAN
+	virtual void AppendSessionToPacket(FNboSerializeToBuffer& Packet, const FSessionLAN& Session) override;
+	virtual void ReadSessionFromPacket(FNboSerializeFromBuffer& Packet, FSessionLAN& Session) override;
 };
 
 /* UE::Online */ }
\ No newline at end of file
