Removal of include cycle

diff --git a/Engine/Plugins/Online/Experimental/OnlineServicesNull/Source/Private/Online/LobbiesNull.cpp b/Engine/Plugins/Online/Experimental/OnlineServicesNull/Source/Private/Online/LobbiesNull.cpp
--- a/Engine/Plugins/Online/Experimental/OnlineServicesNull/Source/Private/Online/LobbiesNull.cpp
+++ b/Engine/Plugins/Online/Experimental/OnlineServicesNull/Source/Private/Online/LobbiesNull.cpp
@@ -3,6 +3,7 @@
 #include "Online/LobbiesNull.h"
 #include "Online/OnlineServicesNull.h"
 #include "Online/OnlineServicesCommonEngineUtils.h"
+#include "Online/NboSerializerNullSvc.h"
 #include "Misc/Guid.h"
 #include "SocketSubsystem.h"
 #include "UObject/CoreNet.h"
diff --git a/Engine/Plugins/Online/Experimental/OnlineServicesNull/Source/Public/Online/LobbiesNull.h b/Engine/Plugins/Online/Experimental/OnlineServicesNull/Source/Public/Online/LobbiesNull.h
--- a/Engine/Plugins/Online/Experimental/OnlineServicesNull/Source/Public/Online/LobbiesNull.h
+++ b/Engine/Plugins/Online/Experimental/OnlineServicesNull/Source/Public/Online/LobbiesNull.h
@@ -5,7 +5,9 @@
 #include "Online/LobbiesCommon.h"
 #include "Online/Lobbies.h"
 #include "Online/LANBeacon.h"
-#include "Online/NboSerializerNullSvc.h"
+
+class FNboSerializeToBuffer;
+class FNboSerializeFromBuffer;
 
 namespace UE::Online {
 
