OSSv2: Add Leaderboards null implementation.

diff --git a/Engine/Plugins/Online/Experimental/OnlineServicesNull/Source/Private/Online/LeaderboardsNull.cpp b/Engine/Plugins/Online/Experimental/OnlineServicesNull/Source/Private/Online/LeaderboardsNull.cpp
new file mode 100644
--- /dev/null
+++ b/Engine/Plugins/Online/Experimental/OnlineServicesNull/Source/Private/Online/LeaderboardsNull.cpp
@@ -0,0 +1,301 @@
+// Copyright Epic Games, Inc. All Rights Reserved.
+
+#include "Online/LeaderboardsNull.h"
+#include "Online/AuthNull.h"
+#include "Online/OnlineServicesNull.h"
+#include "Math/UnrealMathUtility.h"
+
+namespace UE::Online {
+
+struct FFindLeaderboardDataByName
+{
+	FFindLeaderboardDataByName(const FString& InName)
+		: Name(InName)
+	{
+	}
+
+	bool operator()(const FLeaderboardDataNull& LeaderboardData) const
+	{
+		return LeaderboardData.Name == Name;
+	}
+
+	const FString& Name;
+};
+
+FLeaderboardsNull::FLeaderboardsNull(FOnlineServicesNull& InOwningSubsystem)
+	: Super(InOwningSubsystem)
+{
+}
+
+namespace Private
+{
+
+uint64 GetUpdatedScoreByLeaderboardDefinition(FLeaderboardDefinition* LeaderboardDefinition, uint64 OldScore, uint64 NewScore)
+{
+	switch (LeaderboardDefinition->LeaderboardUpdateMethod)
+	{
+	case ELeaderboardUpdateMethod::KeepBest:
+		switch (LeaderboardDefinition->LeaderboardOrderMethod)
+		{
+		case ELeaderboardOrderMethod::Descending:
+			return FMath::Max(NewScore, OldScore);
+		case ELeaderboardOrderMethod::Ascending:
+			return FMath::Min(NewScore, OldScore);
+		default: checkNoEntry();
+			return NewScore;
+		}
+		break;
+	default: checkNoEntry(); // Intentional fallthrough
+	case ELeaderboardUpdateMethod::Force:
+		return NewScore;
+	}
+}
+
+}
+
+TOnlineAsyncOpHandle<FWriteLeaderboardScores> FLeaderboardsNull::WriteLeaderboardScores(FWriteLeaderboardScores::Params&& Params)
+{
+	TOnlineAsyncOpRef<FWriteLeaderboardScores> Op = GetOp<FWriteLeaderboardScores>(MoveTemp(Params));
+
+	if (!Services.Get<FAuthNull>()->IsLoggedIn(Op->GetParams().LocalUserId))
+	{
+		Op->SetError(Errors::InvalidUser());
+		return Op->GetHandle();
+	}
+
+	Op->Then([this](TOnlineAsyncOp<FWriteLeaderboardScores>& InAsyncOp) mutable
+	{
+		FLeaderboardDefinition* LeaderboardDefinition = LeaderboardDefinitions.Find(InAsyncOp.GetParams().BoardName);
+		if (!LeaderboardDefinition)
+		{
+			InAsyncOp.SetError(Errors::NotConfigured());
+			return;
+		}
+
+		FLeaderboardDataNull* LeaderboardData = LeaderboardsData.FindByPredicate(FFindLeaderboardDataByName(InAsyncOp.GetParams().BoardName));
+		if (!LeaderboardData)
+		{
+			LeaderboardData = &LeaderboardsData.Emplace_GetRef();
+			LeaderboardData->Name = InAsyncOp.GetParams().BoardName;
+		}
+
+		TDoubleLinkedList<FUserScoreNull>& UserScoreList = LeaderboardData->UserScoreList;
+
+		uint64 UpdatedScore = InAsyncOp.GetParams().Score;
+
+		// Update score and remove from list
+		TDoubleLinkedList<FUserScoreNull>::TDoubleLinkedListNode* CurrentNode = UserScoreList.GetHead();
+		while (CurrentNode)
+		{
+			const FUserScoreNull& UserScore = CurrentNode->GetValue();
+			if (UserScore.UserId == InAsyncOp.GetParams().LocalUserId)
+			{
+				UpdatedScore = Private::GetUpdatedScoreByLeaderboardDefinition(LeaderboardDefinition, UserScore.Score, InAsyncOp.GetParams().Score);
+				UserScoreList.RemoveNode(CurrentNode);
+				break;
+			}
+
+			CurrentNode = CurrentNode->GetNextNode();
+		}
+
+		// Insert into list by order
+		FUserScoreNull UserScoreToInsert;
+		UserScoreToInsert.Score = UpdatedScore;
+		UserScoreToInsert.UserId = InAsyncOp.GetParams().LocalUserId;
+
+		CurrentNode = UserScoreList.GetHead();
+		while (CurrentNode)
+		{
+			const FUserScoreNull& UserScore = CurrentNode->GetValue();
+
+			bool ShouldInsert = false;
+			switch (LeaderboardDefinition->LeaderboardOrderMethod)
+			{
+			case ELeaderboardOrderMethod::Descending:
+				ShouldInsert = UpdatedScore > UserScore.Score;
+				break;
+			case ELeaderboardOrderMethod::Ascending:
+				ShouldInsert = UpdatedScore < UserScore.Score;
+				break;
+			default: checkNoEntry();
+				break;
+			}
+			if (ShouldInsert)
+			{
+				UserScoreList.InsertNode(UserScoreToInsert, CurrentNode);
+				break;
+			}
+
+			CurrentNode = CurrentNode->GetNextNode();
+			if (!CurrentNode)
+			{
+				UserScoreList.AddTail(UserScoreToInsert);
+			}
+		}
+
+		InAsyncOp.SetResult({});
+	})
+	.Enqueue(GetSerialQueue());
+
+	return Op->GetHandle();
+}
+
+TOnlineAsyncOpHandle<FReadEntriesForUsers> FLeaderboardsNull::ReadEntriesForUsers(FReadEntriesForUsers::Params&& Params)
+{
+	TOnlineAsyncOpRef<FReadEntriesForUsers> Op = GetOp<FReadEntriesForUsers>(MoveTemp(Params));
+
+	if (!Services.Get<FAuthNull>()->IsLoggedIn(Op->GetParams().LocalUserId))
+	{
+		Op->SetError(Errors::InvalidUser());
+		return Op->GetHandle();
+	}
+
+	Op->Then([this](TOnlineAsyncOp<FReadEntriesForUsers>& InAsyncOp) mutable
+	{
+		FReadEntriesForUsers::Result Result;
+
+		if (FLeaderboardDataNull* LeaderboardData = LeaderboardsData.FindByPredicate(FFindLeaderboardDataByName(InAsyncOp.GetParams().BoardName)))
+		{
+			for (const FOnlineAccountIdHandle& UserId : InAsyncOp.GetParams().UserIds)
+			{
+				uint32 Index = 0;
+				TDoubleLinkedList<FUserScoreNull>::TDoubleLinkedListNode* CurrentNode = LeaderboardData->UserScoreList.GetHead();
+				while (CurrentNode)
+				{
+					const FUserScoreNull& UserScore = CurrentNode->GetValue();
+					// When needed, use map/set instead of array for UserIds in Params, to improve performance
+					if (UserScore.UserId == UserId)
+					{
+						FLeaderboardEntry& LeaderboardEntry = Result.Entries.Emplace_GetRef();
+						LeaderboardEntry.UserId = UserId;
+						LeaderboardEntry.Rank = Index;
+						LeaderboardEntry.Score = UserScore.Score;
+						break;
+					}
+					CurrentNode = CurrentNode->GetNextNode();
+					++Index;
+				}
+			}
+		}
+
+		InAsyncOp.SetResult(MoveTemp(Result));
+	})
+	.Enqueue(GetSerialQueue());
+
+	return Op->GetHandle();
+}
+
+namespace Private
+{
+
+void ReadEntriesInRange(TDoubleLinkedList<FUserScoreNull>& UserScoreList, uint32 StartIndex, uint32 Limit, TArray<FLeaderboardEntry>& OutEntries)
+{
+	uint32 Index = 0;
+	TDoubleLinkedList<FUserScoreNull>::TDoubleLinkedListNode* CurrentNode = UserScoreList.GetHead();
+	while (CurrentNode)
+	{
+		if (Index >= StartIndex + Limit)
+			break;
+
+		if (Index >= StartIndex)
+		{
+			const FUserScoreNull& UserScore = CurrentNode->GetValue();
+			FLeaderboardEntry& LeaderboardEntry = OutEntries.Emplace_GetRef();
+			LeaderboardEntry.UserId = UserScore.UserId;
+			LeaderboardEntry.Rank = Index;
+			LeaderboardEntry.Score = UserScore.Score;
+		}
+
+		CurrentNode = CurrentNode->GetNextNode();
+		++Index;
+	}
+}
+
+}
+
+TOnlineAsyncOpHandle<FReadEntriesAroundRank> FLeaderboardsNull::ReadEntriesAroundRank(FReadEntriesAroundRank::Params&& Params)
+{
+	TOnlineAsyncOpRef<FReadEntriesAroundRank> Op = GetOp<FReadEntriesAroundRank>(MoveTemp(Params));
+
+	if (!Services.Get<FAuthNull>()->IsLoggedIn(Op->GetParams().LocalUserId))
+	{
+		Op->SetError(Errors::InvalidUser());
+		return Op->GetHandle();
+	}
+
+	if (Op->GetParams().Limit == 0)
+	{
+		Op->SetError(Errors::InvalidParams());
+		return Op->GetHandle();
+	}
+
+	Op->Then([this](TOnlineAsyncOp<FReadEntriesAroundRank>& InAsyncOp)
+	{
+		FReadEntriesAroundRank::Result Result;
+
+		if (FLeaderboardDataNull* LeaderboardData = LeaderboardsData.FindByPredicate(FFindLeaderboardDataByName(InAsyncOp.GetParams().BoardName)))
+		{
+			Private::ReadEntriesInRange(LeaderboardData->UserScoreList, InAsyncOp.GetParams().Rank, InAsyncOp.GetParams().Limit, Result.Entries);
+		}
+
+		InAsyncOp.SetResult(MoveTemp(Result));
+	})
+	.Enqueue(GetSerialQueue());
+
+	return Op->GetHandle();
+}
+
+TOnlineAsyncOpHandle<FReadEntriesAroundUser> FLeaderboardsNull::ReadEntriesAroundUser(FReadEntriesAroundUser::Params&& Params)
+{
+	TOnlineAsyncOpRef<FReadEntriesAroundUser> Op = GetOp<FReadEntriesAroundUser>(MoveTemp(Params));
+
+	if (!Services.Get<FAuthNull>()->IsLoggedIn(Op->GetParams().LocalUserId))
+	{
+		Op->SetError(Errors::InvalidUser());
+		return Op->GetHandle();
+	}
+
+	if (Op->GetParams().Limit == 0)
+	{
+		Op->SetError(Errors::InvalidParams());
+		return Op->GetHandle();
+	}
+
+	Op->Then([this](TOnlineAsyncOp<FReadEntriesAroundUser>& InAsyncOp)
+	{
+		FReadEntriesAroundUser::Result Result;
+
+		if (FLeaderboardDataNull* LeaderboardData = LeaderboardsData.FindByPredicate(FFindLeaderboardDataByName(InAsyncOp.GetParams().BoardName)))
+		{
+			bool FoundUser = false;
+			uint32 UserRank = 0;
+			TDoubleLinkedList<FUserScoreNull>::TDoubleLinkedListNode* CurrentNode = LeaderboardData->UserScoreList.GetHead();
+			while (CurrentNode)
+			{
+				const FUserScoreNull& UserScore = CurrentNode->GetValue();
+				if (UserScore.UserId == InAsyncOp.GetParams().UserId)
+				{
+					FoundUser = true;
+					break;
+				}
+
+				CurrentNode = CurrentNode->GetNextNode();
+				++UserRank;
+			}
+
+			if (FoundUser)
+			{
+				int32 StartIndex = (int32)UserRank + InAsyncOp.GetParams().Offset;
+				StartIndex = FMath::Max(StartIndex, 0);
+				Private::ReadEntriesInRange(LeaderboardData->UserScoreList, StartIndex, InAsyncOp.GetParams().Limit, Result.Entries);
+			}
+		}
+
+		InAsyncOp.SetResult(MoveTemp(Result));
+	})
+	.Enqueue(GetSerialQueue());
+
+	return Op->GetHandle();
+}
+
+/* UE::Online */ }
diff --git a/Engine/Plugins/Online/Experimental/OnlineServicesNull/Source/Private/Online/OnlineServicesNull.cpp b/Engine/Plugins/Online/Experimental/OnlineServicesNull/Source/Private/Online/OnlineServicesNull.cpp
--- a/Engine/Plugins/Online/Experimental/OnlineServicesNull/Source/Private/Online/OnlineServicesNull.cpp
+++ b/Engine/Plugins/Online/Experimental/OnlineServicesNull/Source/Private/Online/OnlineServicesNull.cpp
@@ -5,9 +5,10 @@
 
 #include "Online/AchievementsNull.h"
 #include "Online/AuthNull.h"
-#include "Online/StatsNull.h"
+#include "Online/LeaderboardsNull.h"
 #include "Online/LobbiesNull.h"
 #include "Online/PresenceNull.h"
+#include "Online/StatsNull.h"
 #include "Online/TitleFileNull.h"
 #include "Online/UserFileNull.h"
 #include "Online/SessionsNull.h"
@@ -39,8 +40,9 @@ void FOnlineServicesNull::RegisterComponents()
 	Components.Register<FAchievementsNull>(*this);
 	Components.Register<FAuthNull>(*this);
 	Components.Register<FPresenceNull>(*this);
-	Components.Register<FStatsNull>(*this);
+	Components.Register<FLeaderboardsNull>(*this);
 	Components.Register<FLobbiesNull>(*this);
+	Components.Register<FStatsNull>(*this);
 	Components.Register<FTitleFileNull>(*this);
 	Components.Register<FUserFileNull>(*this);
 	Components.Register<FSessionsNull>(*this);
diff --git a/Engine/Plugins/Online/Experimental/OnlineServicesNull/Source/Public/Online/LeaderboardsNull.h b/Engine/Plugins/Online/Experimental/OnlineServicesNull/Source/Public/Online/LeaderboardsNull.h
new file mode 100644
--- /dev/null
+++ b/Engine/Plugins/Online/Experimental/OnlineServicesNull/Source/Public/Online/LeaderboardsNull.h
@@ -0,0 +1,44 @@
+// Copyright Epic Games, Inc. All Rights Reserved.
+
+#pragma once
+
+#include "Misc/Optional.h"
+#include "Online/LeaderboardsCommon.h"
+#include "Containers/List.h"
+
+namespace UE::Online {
+
+class FOnlineServicesNull;
+
+struct FUserScoreNull
+{
+	FOnlineAccountIdHandle UserId;
+	uint64 Score;
+};
+
+struct FLeaderboardDataNull
+{
+	FString Name;
+	TDoubleLinkedList<FUserScoreNull> UserScoreList;
+};
+
+class ONLINESERVICESNULL_API FLeaderboardsNull : public FLeaderboardsCommon
+{
+public:
+	using Super = FLeaderboardsCommon;
+
+	FLeaderboardsNull(FOnlineServicesNull& InOwningSubsystem);
+
+	// ILeaderboards
+	virtual TOnlineAsyncOpHandle<FReadEntriesForUsers> ReadEntriesForUsers(FReadEntriesForUsers::Params&& Params) override;
+	virtual TOnlineAsyncOpHandle<FReadEntriesAroundRank> ReadEntriesAroundRank(FReadEntriesAroundRank::Params&& Params) override;
+	virtual TOnlineAsyncOpHandle<FReadEntriesAroundUser> ReadEntriesAroundUser(FReadEntriesAroundUser::Params&& Params) override;
+
+	// FLeaderboardsCommon
+	virtual TOnlineAsyncOpHandle<FWriteLeaderboardScores> WriteLeaderboardScores(FWriteLeaderboardScores::Params&& Params) override;
+
+protected:
+	TArray<FLeaderboardDataNull> LeaderboardsData;
+};
+
+/* UE::Online */ }
