- Fix TitleFileNull not returning InvalidParams from ReadFile when Filename is empty
- Fix TitleFileNull returning wide chars in file contents, by converting to utf8 first.

diff --git a/Engine/Plugins/Online/Experimental/OnlineServicesNull/Source/Private/Online/TitleFileNull.cpp b/Engine/Plugins/Online/Experimental/OnlineServicesNull/Source/Private/Online/TitleFileNull.cpp
--- a/Engine/Plugins/Online/Experimental/OnlineServicesNull/Source/Private/Online/TitleFileNull.cpp
+++ b/Engine/Plugins/Online/Experimental/OnlineServicesNull/Source/Private/Online/TitleFileNull.cpp
@@ -33,10 +33,8 @@ TMap<FString, FTitleFileContentsRef> GetTitleFilesFromConfig()
 
 		if (!FileContentsStr.IsEmpty())
 		{
-			const int ArrayLen = FileContentsStr.GetCharArray().Num() * sizeof(TCHAR);
-			FTitleFileContents FileContents((uint8*)FileContentsStr.GetCharArray().GetData(), ArrayLen);
-
-			Result.Emplace(MoveTemp(Filename), MakeShared<const FTitleFileContents>(MoveTemp(FileContents)));
+			const FTCHARToUTF8 FileContentsUtf8(*FileContentsStr);
+			Result.Emplace(MoveTemp(Filename), MakeShared<FTitleFileContents>((uint8*)FileContentsUtf8.Get(), FileContentsUtf8.Length()));
 		}
 	}
 
@@ -100,6 +98,12 @@ TOnlineAsyncOpHandle<FTitleFileReadFile> FTitleFileNull::ReadFile(FTitleFileRead
 		return Op->GetHandle();
 	}
 
+	if (Op->GetParams().Filename.IsEmpty())
+	{
+		Op->SetError(Errors::InvalidParams());
+		return Op->GetHandle();
+	}
+
 	if (!TitleFiles.IsSet())
 	{
 		// Need to call EnumerateFiles first.
