- Only require Enumerate to have been called for GetEnumeratedFiles, Read/Write/Copy/Delete no longer require it as a prerequisite, to match EOS.

diff --git a/Engine/Plugins/Online/Experimental/OnlineServicesNull/Source/Private/Online/TitleFileNull.cpp b/Engine/Plugins/Online/Experimental/OnlineServicesNull/Source/Private/Online/TitleFileNull.cpp
--- a/Engine/Plugins/Online/Experimental/OnlineServicesNull/Source/Private/Online/TitleFileNull.cpp
+++ b/Engine/Plugins/Online/Experimental/OnlineServicesNull/Source/Private/Online/TitleFileNull.cpp
@@ -8,16 +8,18 @@
 #include "Online/OnlineServicesNull.h"
 #include "Online/OnlineServicesNullTypes.h"
 
-namespace {
+namespace UE::Online {
 
-using FTitleFileContents = UE::Online::FTitleFileContents;
-using FTitleFileContentsRef = UE::Online::FTitleFileContentsRef;
+FTitleFileNull::FTitleFileNull(FOnlineServicesNull& InOwningSubsystem)
+	: Super(InOwningSubsystem)
+{
+}
 
-TMap<FString, FTitleFileContentsRef> GetTitleFilesFromConfig()
+void FTitleFileNull::LoadConfig()
 {
 	const TCHAR* ConfigSection = TEXT("OnlineServices.Null.TitleFile");
 
-	TMap<FString, FTitleFileContentsRef> Result;
+	TitleFiles.Reset();
 
 	for (int FileIdx = 0;; FileIdx++)
 	{
@@ -34,20 +36,9 @@ TMap<FString, FTitleFileContentsRef> GetTitleFilesFromConfig()
 		if (!FileContentsStr.IsEmpty())
 		{
 			const FTCHARToUTF8 FileContentsUtf8(*FileContentsStr);
-			Result.Emplace(MoveTemp(Filename), MakeShared<FTitleFileContents>((uint8*)FileContentsUtf8.Get(), FileContentsUtf8.Length()));
+			TitleFiles.Emplace(MoveTemp(Filename), MakeShared<FTitleFileContents>((uint8*)FileContentsUtf8.Get(), FileContentsUtf8.Length()));
 		}
 	}
-
-	return Result;
-}
-
-}
-
-namespace UE::Online {
-
-FTitleFileNull::FTitleFileNull(FOnlineServicesNull& InOwningSubsystem)
-	: Super(InOwningSubsystem)
-{
 }
 
 TOnlineAsyncOpHandle<FTitleFileEnumerateFiles> FTitleFileNull::EnumerateFiles(FTitleFileEnumerateFiles::Params&& Params)
@@ -60,10 +51,7 @@ TOnlineAsyncOpHandle<FTitleFileEnumerateFiles> FTitleFileNull::EnumerateFiles(FT
 		return Op->GetHandle();
 	}
 
-	if(!TitleFiles.IsSet())
-	{
-		TitleFiles.Emplace(GetTitleFilesFromConfig());
-	}
+	bEnumerated = true;
 
 	Op->SetResult({});
 	return Op->GetHandle();
@@ -76,14 +64,14 @@ TOnlineResult<FTitleFileGetEnumeratedFiles> FTitleFileNull::GetEnumeratedFiles(F
 		return TOnlineResult<FTitleFileGetEnumeratedFiles>(Errors::InvalidUser());
 	}
 
-	if (!TitleFiles.IsSet())
+	if (!bEnumerated)
 	{
 		// Need to call EnumerateFiles first.
 		return TOnlineResult<FTitleFileGetEnumeratedFiles>(Errors::InvalidState());
 	}
 
 	FTitleFileGetEnumeratedFiles::Result Result;
-	TitleFiles->GenerateKeyArray(Result.Filenames);
+	TitleFiles.GenerateKeyArray(Result.Filenames);
 	return TOnlineResult<FTitleFileGetEnumeratedFiles>(MoveTemp(Result));
 	
 }
@@ -104,14 +92,7 @@ TOnlineAsyncOpHandle<FTitleFileReadFile> FTitleFileNull::ReadFile(FTitleFileRead
 		return Op->GetHandle();
 	}
 
-	if (!TitleFiles.IsSet())
-	{
-		// Need to call EnumerateFiles first.
-		Op->SetError(Errors::InvalidState());
-		return Op->GetHandle();
-	}
-
-	const FTitleFileContentsRef* Found = TitleFiles->Find(Op->GetParams().Filename);
+	const FTitleFileContentsRef* Found = TitleFiles.Find(Op->GetParams().Filename);
 	if (!Found)
 	{
 		Op->SetError(Errors::NotFound());
diff --git a/Engine/Plugins/Online/Experimental/OnlineServicesNull/Source/Private/Online/UserFileNull.cpp b/Engine/Plugins/Online/Experimental/OnlineServicesNull/Source/Private/Online/UserFileNull.cpp
--- a/Engine/Plugins/Online/Experimental/OnlineServicesNull/Source/Private/Online/UserFileNull.cpp
+++ b/Engine/Plugins/Online/Experimental/OnlineServicesNull/Source/Private/Online/UserFileNull.cpp
@@ -15,6 +15,34 @@ FUserFileNull::FUserFileNull(FOnlineServicesNull& InOwningSubsystem)
 {
 }
 
+void FUserFileNull::LoadConfig()
+{
+	const TCHAR* ConfigSection = TEXT("OnlineServices.Null.UserFile");
+
+	InitialFileStateFromConfig.Reset();
+
+	for (int FileIdx = 0;; FileIdx++)
+	{
+		FString Filename;
+		GConfig->GetString(ConfigSection, *FString::Printf(TEXT("File_%d_Name"), FileIdx), Filename, GEngineIni);
+		if (Filename.IsEmpty())
+		{
+			break;
+		}
+
+		FString FileContentsStr;
+		GConfig->GetString(ConfigSection, *FString::Printf(TEXT("File_%d_Contents"), FileIdx), FileContentsStr, GEngineIni);
+
+		if (!FileContentsStr.IsEmpty())
+		{
+			const int ArrayLen = FileContentsStr.GetCharArray().Num() * sizeof(TCHAR);
+			FUserFileContents FileContents((uint8*)FileContentsStr.GetCharArray().GetData(), ArrayLen);
+
+			InitialFileStateFromConfig.Emplace(MoveTemp(Filename), MakeShared<const FUserFileContents>(MoveTemp(FileContents)));
+		}
+	}
+}
+
 TOnlineAsyncOpHandle<FUserFileEnumerateFiles> FUserFileNull::EnumerateFiles(FUserFileEnumerateFiles::Params&& InParams)
 {
 	TOnlineAsyncOpRef<FUserFileEnumerateFiles> Op = GetOp<FUserFileEnumerateFiles>(MoveTemp(InParams));
@@ -26,14 +54,8 @@ TOnlineAsyncOpHandle<FUserFileEnumerateFiles> FUserFileNull::EnumerateFiles(FUse
 		return Op->GetHandle();
 	}
 
-	if (!UserToFilesMap.Contains(Params.LocalUserId))
-	{
-		if (!InitialFileStateFromConfig.IsSet())
-		{
-			LoadUserFilesFromConfig();
-		}
-		UserToFilesMap.Emplace(Params.LocalUserId, *InitialFileStateFromConfig);
-	}
+	FUserState& UserState = GetUserState(Params.LocalUserId);
+	UserState.bEnumerated = true;
 
 	Op->SetResult({});
 	return Op->GetHandle();
@@ -46,15 +68,15 @@ TOnlineResult<FUserFileGetEnumeratedFiles> FUserFileNull::GetEnumeratedFiles(FUs
 		return TOnlineResult<FUserFileGetEnumeratedFiles>(Errors::InvalidUser());
 	}
 
-	FUserFileMap* UserFiles = UserToFilesMap.Find(Params.LocalUserId);
-	if (!UserFiles)
+	FUserState& FileState = GetUserState(Params.LocalUserId);
+	if (!FileState.bEnumerated)
 	{
 		// Need to call EnumerateFiles first.
 		return TOnlineResult<FUserFileGetEnumeratedFiles>(Errors::InvalidState());
 	}
 
 	FUserFileGetEnumeratedFiles::Result Result;
-	UserFiles->GenerateKeyArray(Result.Filenames);
+	FileState.Files.GenerateKeyArray(Result.Filenames);
 	return TOnlineResult<FUserFileGetEnumeratedFiles>(MoveTemp(Result));
 }
 
@@ -69,21 +91,14 @@ TOnlineAsyncOpHandle<FUserFileReadFile> FUserFileNull::ReadFile(FUserFileReadFil
 		return Op->GetHandle();
 	}
 
-	FUserFileMap* UserFiles = UserToFilesMap.Find(Params.LocalUserId);
-	if (!UserFiles)
-	{
-		// Need to call EnumerateFiles first.
-		Op->SetError(Errors::InvalidState());
-		return Op->GetHandle();
-	}
-
 	if (Params.Filename.IsEmpty())
 	{
 		Op->SetError(Errors::InvalidParams());
 		return Op->GetHandle();
 	}
 
-	const FUserFileContentsRef* Found = UserFiles->Find(Params.Filename);
+	FUserState& UserState = GetUserState(Params.LocalUserId);
+	const FUserFileContentsRef* Found = UserState.Files.Find(Params.Filename);
 	if (!Found)
 	{
 		Op->SetError(Errors::NotFound());
@@ -103,15 +118,7 @@ TOnlineAsyncOpHandle<FUserFileWriteFile> FUserFileNull::WriteFile(FUserFileWrite
 	{
 		Op->SetError(Errors::InvalidUser());
 		return Op->GetHandle();
-	}
-
-	FUserFileMap* UserFiles = UserToFilesMap.Find(Params.LocalUserId);
-	if (!UserFiles)
-	{
-		// Need to call EnumerateFiles first.
-		Op->SetError(Errors::InvalidState());
-		return Op->GetHandle();
-	}
+	}	
 
 	if (Params.Filename.IsEmpty() || Params.FileContents.IsEmpty())
 	{
@@ -119,7 +126,8 @@ TOnlineAsyncOpHandle<FUserFileWriteFile> FUserFileNull::WriteFile(FUserFileWrite
 		return Op->GetHandle();
 	}
 
-	UserFiles->Emplace(Params.Filename, MakeShared<FUserFileContents>(Params.FileContents));
+	FUserState& UserState = GetUserState(Params.LocalUserId);
+	UserState.Files.Emplace(Params.Filename, MakeShared<FUserFileContents>(Params.FileContents));
 	Op->SetResult({});
 	return Op->GetHandle();
 }
@@ -135,28 +143,21 @@ TOnlineAsyncOpHandle<FUserFileCopyFile> FUserFileNull::CopyFile(FUserFileCopyFil
 		return Op->GetHandle();
 	}
 
-	FUserFileMap* UserFiles = UserToFilesMap.Find(Params.LocalUserId);
-	if (!UserFiles)
-	{
-		// Need to call EnumerateFiles first.
-		Op->SetError(Errors::InvalidState());
-		return Op->GetHandle();
-	}
-
 	if (Params.SourceFilename.IsEmpty() || Params.TargetFilename.IsEmpty())
 	{
 		Op->SetError(Errors::InvalidParams());
 		return Op->GetHandle();
 	}
 
-	const FUserFileContentsRef* SourceFile = UserFiles->Find(Params.SourceFilename);
+	FUserState& UserState = GetUserState(Params.LocalUserId);
+	const FUserFileContentsRef* SourceFile = UserState.Files.Find(Params.SourceFilename);
 	if (!SourceFile)
 	{
 		Op->SetError(Errors::NotFound());
 		return Op->GetHandle();
 	}
 
-	UserFiles->Emplace(Params.TargetFilename, *SourceFile);
+	UserState.Files.Emplace(Params.TargetFilename, *SourceFile);
 	Op->SetResult({});
 	return Op->GetHandle();
 }
@@ -172,21 +173,14 @@ TOnlineAsyncOpHandle<FUserFileDeleteFile> FUserFileNull::DeleteFile(FUserFileDel
 		return Op->GetHandle();
 	}
 
-	FUserFileMap* UserFiles = UserToFilesMap.Find(Params.LocalUserId);
-	if (!UserFiles)
-	{
-		// Need to call EnumerateFiles first.
-		Op->SetError(Errors::InvalidState());
-		return Op->GetHandle();
-	}
-
 	if (Params.Filename.IsEmpty())
 	{
 		Op->SetError(Errors::InvalidParams());
 		return Op->GetHandle();
 	}
 
-	const bool bRemoved = UserFiles->Remove(Params.Filename) != 0;
+	FUserState& UserState = GetUserState(Params.LocalUserId);
+	const bool bRemoved = UserState.Files.Remove(Params.Filename) != 0;
 	if (!bRemoved)
 	{
 		Op->SetError(Errors::NotFound());
@@ -197,32 +191,18 @@ TOnlineAsyncOpHandle<FUserFileDeleteFile> FUserFileNull::DeleteFile(FUserFileDel
 	return Op->GetHandle();
 }
 
-void FUserFileNull::LoadUserFilesFromConfig()
+FUserFileNull::FUserState& FUserFileNull::GetUserState(const FOnlineAccountIdHandle AccountId)
 {
-	const TCHAR* ConfigSection = TEXT("OnlineServices.Null.UserFile");
+	check(AccountId.IsValid());
 
-	InitialFileStateFromConfig.Emplace();
-
-	for (int FileIdx = 0;; FileIdx++)
+	FUserState* FileState = UserStates.Find(AccountId);
+	if (!FileState)
 	{
-		FString Filename;
-		GConfig->GetString(ConfigSection, *FString::Printf(TEXT("File_%d_Name"), FileIdx), Filename, GEngineIni);
-		if (Filename.IsEmpty())
-		{
-			break;
-		}
-
-		FString FileContentsStr;
-		GConfig->GetString(ConfigSection, *FString::Printf(TEXT("File_%d_Contents"), FileIdx), FileContentsStr, GEngineIni);
-
-		if (!FileContentsStr.IsEmpty())
-		{
-			const int ArrayLen = FileContentsStr.GetCharArray().Num() * sizeof(TCHAR);
-			FUserFileContents FileContents((uint8*)FileContentsStr.GetCharArray().GetData(), ArrayLen);
-
-			InitialFileStateFromConfig->Emplace(MoveTemp(Filename), MakeShared<const FUserFileContents>(MoveTemp(FileContents)));
-		}
+		FileState = &UserStates.Emplace(AccountId);
+		FileState->Files = InitialFileStateFromConfig;
 	}
+
+	return *FileState;
 }
 
 /* UE::Online */ }
diff --git a/Engine/Plugins/Online/Experimental/OnlineServicesNull/Source/Public/Online/TitleFileNull.h b/Engine/Plugins/Online/Experimental/OnlineServicesNull/Source/Public/Online/TitleFileNull.h
--- a/Engine/Plugins/Online/Experimental/OnlineServicesNull/Source/Public/Online/TitleFileNull.h
+++ b/Engine/Plugins/Online/Experimental/OnlineServicesNull/Source/Public/Online/TitleFileNull.h
@@ -16,13 +16,18 @@ public:
 
 	FTitleFileNull(FOnlineServicesNull& InOwningSubsystem);
 
+	// IOnlineComponent
+	virtual void LoadConfig() override;
+
 	// ITitleFile
 	virtual TOnlineAsyncOpHandle<FTitleFileEnumerateFiles> EnumerateFiles(FTitleFileEnumerateFiles::Params&& Params) override;
 	virtual TOnlineResult<FTitleFileGetEnumeratedFiles> GetEnumeratedFiles(FTitleFileGetEnumeratedFiles::Params&& Params) override;
 	virtual TOnlineAsyncOpHandle<FTitleFileReadFile> ReadFile(FTitleFileReadFile::Params&& Params) override;
 
 protected:
-	TOptional<TMap<FString, FTitleFileContentsRef>> TitleFiles;
+	TMap<FString, FTitleFileContentsRef> TitleFiles;
+
+	bool bEnumerated = false;
 };
 
 /* UE::Online */ }
diff --git a/Engine/Plugins/Online/Experimental/OnlineServicesNull/Source/Public/Online/UserFileNull.h b/Engine/Plugins/Online/Experimental/OnlineServicesNull/Source/Public/Online/UserFileNull.h
--- a/Engine/Plugins/Online/Experimental/OnlineServicesNull/Source/Public/Online/UserFileNull.h
+++ b/Engine/Plugins/Online/Experimental/OnlineServicesNull/Source/Public/Online/UserFileNull.h
@@ -16,6 +16,9 @@ public:
 
 	FUserFileNull(FOnlineServicesNull& InOwningSubsystem);
 
+	// IOnlineComponent
+	virtual void LoadConfig() override;
+
 	// IUserFile
 	virtual TOnlineAsyncOpHandle<FUserFileEnumerateFiles> EnumerateFiles(FUserFileEnumerateFiles::Params&& Params) override;
 	virtual TOnlineResult<FUserFileGetEnumeratedFiles> GetEnumeratedFiles(FUserFileGetEnumeratedFiles::Params&& Params) override;
@@ -27,10 +30,16 @@ public:
 protected:
 	using FUserFileMap = TMap<FString, FUserFileContentsRef>;
 
-	TMap<FOnlineAccountIdHandle, FUserFileMap> UserToFilesMap;
+	struct FUserState
+	{
+		bool bEnumerated = false;
+		FUserFileMap Files;
+	};
+
+	TMap<FOnlineAccountIdHandle, FUserState> UserStates;
+	FUserState& GetUserState(const FOnlineAccountIdHandle AccountId);
 
-	TOptional<FUserFileMap> InitialFileStateFromConfig;
-	void LoadUserFilesFromConfig();
+	FUserFileMap InitialFileStateFromConfig;
 };
 
 /* UE::Online */ }
