- Update AchievementsNull to use struct config loading

diff --git a/Engine/Plugins/Online/Experimental/OnlineServicesNull/Source/Private/Online/AchievementsNull.cpp b/Engine/Plugins/Online/Experimental/OnlineServicesNull/Source/Private/Online/AchievementsNull.cpp
--- a/Engine/Plugins/Online/Experimental/OnlineServicesNull/Source/Private/Online/AchievementsNull.cpp
+++ b/Engine/Plugins/Online/Experimental/OnlineServicesNull/Source/Private/Online/AchievementsNull.cpp
@@ -12,61 +12,15 @@
 
 namespace UE::Online {
 
-namespace {
-
-TMap<FString, FAchievementDefinition> GetAchievementDefinitionsFromConfig()
+FAchievementsNull::FAchievementsNull(FOnlineServicesNull& InOwningSubsystem)
+	: Super(InOwningSubsystem)
 {
-	const TCHAR* ConfigSection = TEXT("OnlineServices.Null.Achievements");
-
-	TMap<FString, FAchievementDefinition> Result;
-
-	for (int AchievIdx = 0;; AchievIdx++)
-	{
-		FString AchievementId;
-		GConfig->GetString(ConfigSection, *FString::Printf(TEXT("AchievementDef_%d_Id"), AchievIdx), AchievementId, GEngineIni);
-		if (AchievementId.IsEmpty())
-		{
-			break;
-		}
-
-		FAchievementDefinition& AchievementDefinition = Result.Emplace(AchievementId);
-		AchievementDefinition.AchievementId = MoveTemp(AchievementId);
-
-		GConfig->GetText(ConfigSection, *FString::Printf(TEXT("AchievementDef_%d_UnlockedDisplayName"), AchievIdx), AchievementDefinition.UnlockedDisplayName, GEngineIni);
-		GConfig->GetText(ConfigSection, *FString::Printf(TEXT("AchievementDef_%d_UnlockedDescription"), AchievIdx), AchievementDefinition.UnlockedDescription, GEngineIni);
-		GConfig->GetText(ConfigSection, *FString::Printf(TEXT("AchievementDef_%d_LockedDisplayName"), AchievIdx), AchievementDefinition.LockedDisplayName, GEngineIni);
-		GConfig->GetText(ConfigSection, *FString::Printf(TEXT("AchievementDef_%d_LockedDescription"), AchievIdx), AchievementDefinition.LockedDescription, GEngineIni);
-		GConfig->GetText(ConfigSection, *FString::Printf(TEXT("AchievementDef_%d_FlavorText"), AchievIdx), AchievementDefinition.FlavorText, GEngineIni);
-		GConfig->GetString(ConfigSection, *FString::Printf(TEXT("AchievementDef_%d_UnlockedIconUrl"), AchievIdx), AchievementDefinition.UnlockedIconUrl, GEngineIni);
-		GConfig->GetString(ConfigSection, *FString::Printf(TEXT("AchievementDef_%d_LockedIconUrl"), AchievIdx), AchievementDefinition.LockedIconUrl, GEngineIni);
-		GConfig->GetBool(ConfigSection, *FString::Printf(TEXT("AchievementDef_%d_bIsHidden"), AchievIdx), AchievementDefinition.bIsHidden, GEngineIni);
-
-		for (int StatDefIdx = 0;; StatDefIdx++)
-		{
-			FString StatId;
-			GConfig->GetString(ConfigSection, *FString::Printf(TEXT("AchievementDef_%d_StatDef_%d_Id"), AchievIdx, StatDefIdx), StatId, GEngineIni);
-			if (StatId.IsEmpty())
-			{
-				break;
-			}
-
-			FAchievementStatDefinition StatDefinition = AchievementDefinition.StatDefinitions.Emplace_GetRef();
-			StatDefinition.StatId = MoveTemp(StatId);
-
-			int32 UnlockThreshold = 0;
-			GConfig->GetInt(ConfigSection, *FString::Printf(TEXT("AchievementDef_%d_StatDef_%d_UnlockThreshold"), AchievIdx, StatDefIdx), UnlockThreshold, GEngineIni);
-			StatDefinition.UnlockThreshold = UnlockThreshold;
-		}
-	}
-
-	return Result;
-}
-
 }
 
-FAchievementsNull::FAchievementsNull(FOnlineServicesNull& InOwningSubsystem)
-	: Super(InOwningSubsystem)
+void FAchievementsNull::UpdateConfig()
 {
+	Super::UpdateConfig();
+	TOnlineComponent::LoadConfig(Config);
 }
 
 TOnlineAsyncOpHandle<FQueryAchievementDefinitions> FAchievementsNull::QueryAchievementDefinitions(FQueryAchievementDefinitions::Params&& Params)
@@ -79,10 +33,7 @@ TOnlineAsyncOpHandle<FQueryAchievementDefinitions> FAchievementsNull::QueryAchie
 		return Op->GetHandle();
 	}
 
-	if(!AchievementDefinitions.IsSet())
-	{
-		AchievementDefinitions.Emplace(GetAchievementDefinitionsFromConfig());
-	}
+	bAchievementDefinitionsQueried = true;
 
 	Op->SetResult({});
 	return Op->GetHandle();
@@ -95,14 +46,17 @@ TOnlineResult<FGetAchievementIds> FAchievementsNull::GetAchievementIds(FGetAchie
 		return TOnlineResult<FGetAchievementIds>(Errors::InvalidUser());
 	}
 
-	if (!AchievementDefinitions.IsSet())
+	if (!bAchievementDefinitionsQueried)
 	{
 		// Call QueryAchievementDefinitions first
 		return TOnlineResult<FGetAchievementIds>(Errors::InvalidState());
 	}
 
 	FGetAchievementIds::Result Result;
-	AchievementDefinitions->GenerateKeyArray(Result.AchievementIds);
+	for(const FAchievementDefinition& Definition : Config.AchievementDefinitions)
+	{
+		Result.AchievementIds.Emplace(Definition.AchievementId);
+	}
 	return TOnlineResult<FGetAchievementIds>(MoveTemp(Result));
 }
 
@@ -113,13 +67,13 @@ TOnlineResult<FGetAchievementDefinition> FAchievementsNull::GetAchievementDefini
 		return TOnlineResult<FGetAchievementDefinition>(Errors::InvalidUser());
 	}
 
-	if (!AchievementDefinitions.IsSet())
+	if (!bAchievementDefinitionsQueried)
 	{
 		// Should call QueryAchievementDefinitions first
 		return TOnlineResult<FGetAchievementDefinition>(Errors::InvalidState());
 	}
 
-	const FAchievementDefinition* AchievementDefinition = AchievementDefinitions->Find(Params.AchievementId);
+	const FAchievementDefinition* AchievementDefinition = FindAchievementDefinition(Params.AchievementId);
 	if (!AchievementDefinition)
 	{
 		return TOnlineResult<FGetAchievementDefinition>(Errors::NotFound());
@@ -138,7 +92,7 @@ TOnlineAsyncOpHandle<FQueryAchievementStates> FAchievementsNull::QueryAchievemen
 		return Op->GetHandle();
 	}
 
-	if (!AchievementDefinitions.IsSet())
+	if (!bAchievementDefinitionsQueried)
 	{
 		// Call QueryAchievementDefinitions first
 		Op->SetError(Errors::InvalidState());
@@ -148,11 +102,10 @@ TOnlineAsyncOpHandle<FQueryAchievementStates> FAchievementsNull::QueryAchievemen
 	if (!AchievementStates.Contains(Op->GetParams().LocalAccountId))
 	{
 		FAchievementStateMap& LocalUserAchievementStates = AchievementStates.Emplace(Op->GetParams().LocalAccountId);
-		for (const TPair<FString, FAchievementDefinition>& Pair : *AchievementDefinitions)
+		for (const FAchievementDefinition& AchievementDefinition: Config.AchievementDefinitions)
 		{
-			const FString& AchievementId = Pair.Key;
-			FAchievementState& AchievementState = LocalUserAchievementStates.Emplace(AchievementId);
-			AchievementState.AchievementId = AchievementId;
+			FAchievementState& AchievementState = LocalUserAchievementStates.Emplace(AchievementDefinition.AchievementId);
+			AchievementState.AchievementId = AchievementDefinition.AchievementId;
 		}
 	}
 
@@ -259,7 +212,7 @@ TOnlineResult<FDisplayAchievementUI> FAchievementsNull::DisplayAchievementUI(FDi
 
 	// Safe to assume they called QueryAchievementDefinitions from this point, as that is a prereq for QueryAchievementStates.
 
-	const FAchievementDefinition* AchievementDefinition = AchievementDefinitions->Find(Params.AchievementId);
+	const FAchievementDefinition* AchievementDefinition = FindAchievementDefinition(Params.AchievementId);
 	if (!AchievementDefinition)
 	{
 		return TOnlineResult<FDisplayAchievementUI>(Errors::NotFound());
@@ -273,4 +226,9 @@ TOnlineResult<FDisplayAchievementUI> FAchievementsNull::DisplayAchievementUI(FDi
 	return TOnlineResult<FDisplayAchievementUI>(FDisplayAchievementUI::Result());
 }
 
+const FAchievementDefinition* FAchievementsNull::FindAchievementDefinition(const FString& AchievementId) const
+{
+	return Config.AchievementDefinitions.FindByPredicate([&AchievementId](const FAchievementDefinition& Def) { return Def.AchievementId == AchievementId; });
+}
+
 /* UE::Online */ }
diff --git a/Engine/Plugins/Online/Experimental/OnlineServicesNull/Source/Public/Online/AchievementsNull.h b/Engine/Plugins/Online/Experimental/OnlineServicesNull/Source/Public/Online/AchievementsNull.h
--- a/Engine/Plugins/Online/Experimental/OnlineServicesNull/Source/Public/Online/AchievementsNull.h
+++ b/Engine/Plugins/Online/Experimental/OnlineServicesNull/Source/Public/Online/AchievementsNull.h
@@ -9,6 +9,19 @@ namespace UE::Online {
 
 class FOnlineServicesNull;
 
+struct FAchievementsNullConfig
+{
+	TArray<FAchievementDefinition> AchievementDefinitions;
+};
+
+namespace Meta {
+
+BEGIN_ONLINE_STRUCT_META(FAchievementsNullConfig)
+	ONLINE_STRUCT_FIELD(FAchievementsNullConfig, AchievementDefinitions)
+END_ONLINE_STRUCT_META()
+
+/* Meta */ }
+
 class ONLINESERVICESNULL_API FAchievementsNull : public FAchievementsCommon
 {
 public:
@@ -16,7 +29,10 @@ public:
 
 	FAchievementsNull(FOnlineServicesNull& InOwningSubsystem);
 
-	//IAchievements
+	// IOnlineComponent
+	virtual void UpdateConfig() override;
+
+	// IAchievements
 	virtual TOnlineAsyncOpHandle<FQueryAchievementDefinitions> QueryAchievementDefinitions(FQueryAchievementDefinitions::Params&& Params) override;
 	virtual TOnlineResult<FGetAchievementIds> GetAchievementIds(FGetAchievementIds::Params&& Params) override;
 	virtual TOnlineResult<FGetAchievementDefinition> GetAchievementDefinition(FGetAchievementDefinition::Params&& Params) override;
@@ -29,8 +45,12 @@ protected:
 	using FAchievementDefinitionMap = TMap<FString, FAchievementDefinition>;
 	using FAchievementStateMap = TMap<FString, FAchievementState>;
 
-	TOptional<FAchievementDefinitionMap> AchievementDefinitions;
+	bool bAchievementDefinitionsQueried = false;
 	TMap<FAccountId, FAchievementStateMap> AchievementStates;
+
+	FAchievementsNullConfig Config;
+
+	const FAchievementDefinition* FindAchievementDefinition(const FString& AchievementId) const;
 };
 
 /* UE::Online */ }
