- Structify TitleFileNull config

diff --git a/Engine/Plugins/Online/Experimental/OnlineServicesNull/Source/Private/Online/TitleFileNull.cpp b/Engine/Plugins/Online/Experimental/OnlineServicesNull/Source/Private/Online/TitleFileNull.cpp
--- a/Engine/Plugins/Online/Experimental/OnlineServicesNull/Source/Private/Online/TitleFileNull.cpp
+++ b/Engine/Plugins/Online/Experimental/OnlineServicesNull/Source/Private/Online/TitleFileNull.cpp
@@ -10,6 +10,30 @@
 
 namespace UE::Online {
 
+struct FTitleFileNullFileDefinition
+{
+	FString Name;
+	FString Contents;
+};
+
+struct FTitleFileNullConfig
+{
+	TArray<FTitleFileNullFileDefinition> Files;
+};
+
+namespace Meta {
+
+BEGIN_ONLINE_STRUCT_META(FTitleFileNullFileDefinition)
+	ONLINE_STRUCT_FIELD(FTitleFileNullFileDefinition, Name),
+	ONLINE_STRUCT_FIELD(FTitleFileNullFileDefinition, Contents)
+END_ONLINE_STRUCT_META()
+
+BEGIN_ONLINE_STRUCT_META(FTitleFileNullConfig)
+	ONLINE_STRUCT_FIELD(FTitleFileNullConfig, Files)
+END_ONLINE_STRUCT_META()
+
+/* Meta */ }
+
 FTitleFileNull::FTitleFileNull(FOnlineServicesNull& InOwningSubsystem)
 	: Super(InOwningSubsystem)
 {
@@ -17,27 +41,17 @@ FTitleFileNull::FTitleFileNull(FOnlineServicesNull& InOwningSubsystem)
 
 void FTitleFileNull::UpdateConfig()
 {
-	const TCHAR* ConfigSection = TEXT("OnlineServices.Null.TitleFile");
+	Super::UpdateConfig();
+
+	FTitleFileNullConfig Config;
+	TOnlineComponent::LoadConfig(Config);
 
 	TitleFiles.Reset();
 
-	for (int FileIdx = 0;; FileIdx++)
+	for (FTitleFileNullFileDefinition& File : Config.Files)
 	{
-		FString Filename;
-		GConfig->GetString(ConfigSection, *FString::Printf(TEXT("File_%d_Name"), FileIdx), Filename, GEngineIni);
-		if (Filename.IsEmpty())
-		{
-			break;
-		}
-
-		FString FileContentsStr;
-		GConfig->GetString(ConfigSection, *FString::Printf(TEXT("File_%d_Contents"), FileIdx), FileContentsStr, GEngineIni);
-
-		if (!FileContentsStr.IsEmpty())
-		{
-			const FTCHARToUTF8 FileContentsUtf8(*FileContentsStr);
-			TitleFiles.Emplace(MoveTemp(Filename), MakeShared<FTitleFileContents>((uint8*)FileContentsUtf8.Get(), FileContentsUtf8.Length()));
-		}
+		const FTCHARToUTF8 FileContentsUtf8(*File.Contents);
+		TitleFiles.Emplace(MoveTemp(File.Name), MakeShared<FTitleFileContents>((uint8*)FileContentsUtf8.Get(), FileContentsUtf8.Length()));
 	}
 }
 
