Refactor FOnlineAccountIdRegistryNull to not store user association data so that it is stable across service restarts.
Add missing RW locks in FAccountInfoRegistryNULL

diff --git a/Engine/Plugins/Online/Experimental/OnlineServicesNull/Source/Private/Online/AuthNull.cpp b/Engine/Plugins/Online/Experimental/OnlineServicesNull/Source/Private/Online/AuthNull.cpp
--- a/Engine/Plugins/Online/Experimental/OnlineServicesNull/Source/Private/Online/AuthNull.cpp
+++ b/Engine/Plugins/Online/Experimental/OnlineServicesNull/Source/Private/Online/AuthNull.cpp
@@ -6,6 +6,7 @@
 #include "Algo/Transform.h"
 #include "Misc/CoreDelegates.h"
 #include "Misc/OutputDeviceRedirector.h"
+#include "Misc/ScopeRWLock.h"
 #include "Online/OnlineServicesNull.h"
 #include "Online/OnlineServicesNullTypes.h"
 #include "Online/AuthErrors.h"
@@ -76,7 +77,7 @@ TSharedRef<FAccountInfoNull> CreateAccountInfo(const FAuthNullConfig& Config, FP
 {
 	const FString DisplayId = GenerateRandomUserId(Config, PlatformUserId);
 	return MakeShared<FAccountInfoNull>(FAccountInfoNull{ {
-		FOnlineAccountIdRegistryNull::Get().Create(DisplayId, PlatformUserId),
+		FOnlineAccountIdRegistryNull::Get().FindOrAddAccountId(DisplayId),
 		PlatformUserId,
 		ELoginStatus::LoggedIn,
 		{ { AccountAttributeData::DisplayName, DisplayId } }
@@ -97,6 +98,7 @@ TSharedPtr<FAccountInfoNull> FAccountInfoRegistryNULL::Find(FAccountId AccountId
 
 void FAccountInfoRegistryNULL::Register(const TSharedRef<FAccountInfoNull>& AccountInfoNULL)
 {
+	FWriteScopeLock Lock(IndexLock);
 	DoRegister(AccountInfoNULL);
 }
 
@@ -104,6 +106,7 @@ void FAccountInfoRegistryNULL::Unregister(FAccountId AccountId)
 {
 	if (TSharedPtr<FAccountInfoNull> AccountInfoNULL = Find(AccountId))
 	{
+		FWriteScopeLock Lock(IndexLock);
 		DoUnregister(AccountInfoNULL.ToSharedRef());
 	}
 	else
@@ -177,37 +180,14 @@ FOnlineAccountIdRegistryNull& FOnlineAccountIdRegistryNull::Get()
 	return Instance;
 }
 
-FAccountId FOnlineAccountIdRegistryNull::Find(FString AccountId) const
+FAccountId FOnlineAccountIdRegistryNull::Find(const FString& AccountId) const
 {
-	const FOnlineAccountIdString* Entry = StringToId.Find(AccountId);
-	if(Entry)
-	{
-		return Entry->AccountId;
-	}
-	return FAccountId();
-}
-
-FAccountId FOnlineAccountIdRegistryNull::Find(FPlatformUserId UserId) const
-{
-	const FOnlineAccountIdString* Entry = LocalUserMap.Find(UserId);
-	if (Entry)
-	{
-		return Entry->AccountId;
-	}
-	return FAccountId();
-}
-
-FAccountId FOnlineAccountIdRegistryNull::Find(int32 UserIndex) const
-{
-	const FOnlineAccountIdString* Entry = LocalUserMap.Find(FPlatformMisc::GetPlatformUserForUserIndex(UserIndex));
-	if (Entry)
-	{
-		return Entry->AccountId;
-	}
-	return FAccountId();
+	const FReadScopeLock ReadLock(Lock);
+	const FOnlineAccountIdString* Entry = FindNoLock(AccountId);
+	return Entry ? Entry->AccountId : FAccountId();
 }
 
-const FOnlineAccountIdString* FOnlineAccountIdRegistryNull::GetInternal(const FAccountId& AccountId) const
+const FOnlineAccountIdString* FOnlineAccountIdRegistryNull::FindNoLock(const FAccountId& AccountId) const
 {
 	if(AccountId.IsValid() && AccountId.GetOnlineServicesType() == EOnlineServices::Null && AccountId.GetHandle() <= (uint32)Ids.Num())
 	{
@@ -216,35 +196,45 @@ const FOnlineAccountIdString* FOnlineAccountIdRegistryNull::GetInternal(const FA
 	return nullptr;
 }
 
-FAccountId FOnlineAccountIdRegistryNull::Create(FString AccountId, FPlatformUserId PlatformUserId/* = PLATFORMUSERID_NONE*/)
+const FOnlineAccountIdString* FOnlineAccountIdRegistryNull::FindNoLock(const FString& AccountId) const
+{
+	FOnlineAccountIdString* const* Entry = StringToIdIndex.Find(AccountId);
+	return Entry ? *Entry : nullptr;
+}
+
+FAccountId FOnlineAccountIdRegistryNull::FindOrAddAccountId(const FString& AccountId)
 {
-	FAccountId ExistingAccountId = Find(AccountId);
-	if(ExistingAccountId.IsValid())
+	// Check for existing entry under read lock.
 	{
-		UE_LOG(LogOnlineServices, Error, TEXT("[FOnlineAccountIdRegistryNull::Create] Found a duplicate ID for local user %d."), FPlatformMisc::GetUserIndexForPlatformUser(PlatformUserId));
-		return ExistingAccountId;
+		const FReadScopeLock ReadLock(Lock);
+		if (const FOnlineAccountIdString* ExistingAccountId = FindNoLock(AccountId))
+		{
+			return ExistingAccountId->AccountId;
+		}
 	}
 
-	if (!PlatformUserId.IsValid())
+	// Check for existing entry again under write lock before adding entry.
 	{
-		UE_LOG(LogOnlineServices, Warning, TEXT("[FOnlineAccountIdRegistryNull::Create] Unable to create id: PlatformUserId is invalid."));
-		return FAccountId();
-	}
+		const FWriteScopeLock WriteLock(Lock);
+		if (const FOnlineAccountIdString* ExistingAccountId = FindNoLock(AccountId))
+		{
+			return ExistingAccountId->AccountId;
+		}
 
-	FOnlineAccountIdString& Id = Ids.Emplace_GetRef();
-	Id.AccountIndex = Ids.Num();
-	Id.Data = AccountId;
-	Id.AccountId = FAccountId(EOnlineServices::Null, Id.AccountIndex);
-	
-	StringToId.Add(AccountId, Id);
-	LocalUserMap.Add(PlatformUserId, Id);
+		FOnlineAccountIdString& Id = Ids.Emplace_GetRef();
+		Id.AccountIndex = Ids.Num();
+		Id.Data = AccountId;
+		Id.AccountId = FAccountId(EOnlineServices::Null, Id.AccountIndex);
 
-	return Id.AccountId;
+		StringToIdIndex.Add(AccountId, &Id);
+		return Id.AccountId;
+	}
 }
 
 FString FOnlineAccountIdRegistryNull::ToLogString(const FAccountId& AccountId) const
 {
-	if(const FOnlineAccountIdString* Id = GetInternal(AccountId))
+	const FReadScopeLock ReadLock(Lock);
+	if(const FOnlineAccountIdString* Id = FindNoLock(AccountId))
 	{
 		return Id->Data;
 	}
@@ -255,7 +245,8 @@ FString FOnlineAccountIdRegistryNull::ToLogString(const FAccountId& AccountId) c
 
 TArray<uint8> FOnlineAccountIdRegistryNull::ToReplicationData(const FAccountId& AccountId) const
 {
-	if (const FOnlineAccountIdString* Id = GetInternal(AccountId))
+	const FReadScopeLock ReadLock(Lock);
+	if (const FOnlineAccountIdString* Id = FindNoLock(AccountId))
 	{
 		TArray<uint8> ReplicationData;
 		ReplicationData.SetNumUninitialized(Id->Data.Len());
@@ -272,7 +263,7 @@ FAccountId FOnlineAccountIdRegistryNull::FromReplicationData(const TArray<uint8>
 	FString Result = BytesToString(ReplicationData.GetData(), ReplicationData.Num());
 	if(Result.Len() > 0)
 	{
-		return Create(Result);
+		return FindOrAddAccountId(Result);
 	}
 	return FAccountId();
 }
diff --git a/Engine/Plugins/Online/Experimental/OnlineServicesNull/Source/Public/Online/AuthNull.h b/Engine/Plugins/Online/Experimental/OnlineServicesNull/Source/Public/Online/AuthNull.h
--- a/Engine/Plugins/Online/Experimental/OnlineServicesNull/Source/Public/Online/AuthNull.h
+++ b/Engine/Plugins/Online/Experimental/OnlineServicesNull/Source/Public/Online/AuthNull.h
@@ -3,6 +3,7 @@
 #pragma once
 
 #include "CoreMinimal.h"
+#include "HAL/CriticalSection.h"
 #include "Online/AuthCommon.h"
 
 namespace UE::Online {
@@ -22,11 +23,8 @@ class FOnlineAccountIdRegistryNull : public IOnlineAccountIdRegistry
 public:
 	static FOnlineAccountIdRegistryNull& Get();
 
-	FAccountId Find(FString AccountId) const;
-	FAccountId Find(FPlatformUserId UserId) const;
-	FAccountId Find(int32 UserIndex) const;
-
-	FAccountId Create(FString AccountId, FPlatformUserId PlatformUserId = PLATFORMUSERID_NONE);
+	FAccountId Find(const FString& AccountId) const;
+	FAccountId FindOrAddAccountId(const FString& AccountId);
 
 	// Begin IOnlineAccountIdRegistry
 	virtual FString ToLogString(const FAccountId& AccountId) const override;
@@ -37,12 +35,13 @@ public:
 	virtual ~FOnlineAccountIdRegistryNull() = default;
 
 private:
-	const FOnlineAccountIdString* GetInternal(const FAccountId& AccountId) const;
+	// Retrieve an entry with no lock. For internal use only.
+	const FOnlineAccountIdString* FindNoLock(const FAccountId& AccountId) const;
+	const FOnlineAccountIdString* FindNoLock(const FString& AccountId) const;
 
-	// how much of these are actually necessary?
+	mutable FRWLock Lock;
 	TArray<FOnlineAccountIdString> Ids;
-	TMap<FString, FOnlineAccountIdString> StringToId;
-	TMap<FPlatformUserId, FOnlineAccountIdString> LocalUserMap;
+	TMap<FString, FOnlineAccountIdString*> StringToIdIndex;
 
 };
 
diff --git a/Engine/Plugins/Online/Experimental/OnlineServicesOSSAdapter/Source/Private/Online/AuthOSSAdapter.cpp b/Engine/Plugins/Online/Experimental/OnlineServicesOSSAdapter/Source/Private/Online/AuthOSSAdapter.cpp
--- a/Engine/Plugins/Online/Experimental/OnlineServicesOSSAdapter/Source/Private/Online/AuthOSSAdapter.cpp
+++ b/Engine/Plugins/Online/Experimental/OnlineServicesOSSAdapter/Source/Private/Online/AuthOSSAdapter.cpp
@@ -2,6 +2,7 @@
 
 #include "Online/AuthOSSAdapter.h"
 
+#include "Misc/ScopeRWLock.h"
 #include "Online/AuthErrors.h"
 #include "Online/OnlineServicesOSSAdapter.h"
 #include "Online/OnlineIdOSSAdapter.h"
@@ -107,6 +108,7 @@ TSharedPtr<FAccountInfoOSSAdapter> FAccountInfoRegistryOSSAdapter::Find(FAccount
 
 void FAccountInfoRegistryOSSAdapter::Register(const TSharedRef<FAccountInfoOSSAdapter>& AccountInfoNULL)
 {
+	FWriteScopeLock Lock(IndexLock);
 	DoRegister(AccountInfoNULL);
 }
 
@@ -114,6 +116,7 @@ void FAccountInfoRegistryOSSAdapter::Unregister(FAccountId AccountId)
 {
 	if (TSharedPtr<FAccountInfoOSSAdapter> AccountInfoNULL = Find(AccountId))
 	{
+		FWriteScopeLock Lock(IndexLock);
 		DoUnregister(AccountInfoNULL.ToSharedRef());
 	}
 	else
