Fix for Horde Issue #219110

diff --git a/Engine/Plugins/Online/OnlineServicesEOS/Source/Private/OnlineServicesEOSModule.cpp b/Engine/Plugins/Online/OnlineServicesEOS/Source/Private/OnlineServicesEOSModule.cpp
--- a/Engine/Plugins/Online/OnlineServicesEOS/Source/Private/OnlineServicesEOSModule.cpp
+++ b/Engine/Plugins/Online/OnlineServicesEOS/Source/Private/OnlineServicesEOSModule.cpp
@@ -38,6 +38,7 @@ void FOnlineServicesEOSModule::StartupModule()
 		FModuleManager::Get().LoadModuleChecked(EOSSharedModuleName);
 	}
 
+	// Making sure we load the module at this point will avoid errors while cooking
 	const FName OnlineServicesInterfaceModuleName = TEXT("OnlineServicesInterface");
 	if (!FModuleManager::Get().IsModuleLoaded(OnlineServicesInterfaceModuleName))
 	{
diff --git a/Engine/Plugins/Online/OnlineServicesEOSGS/Source/Private/OnlineServicesEOSGSModule.cpp b/Engine/Plugins/Online/OnlineServicesEOSGS/Source/Private/OnlineServicesEOSGSModule.cpp
--- a/Engine/Plugins/Online/OnlineServicesEOSGS/Source/Private/OnlineServicesEOSGSModule.cpp
+++ b/Engine/Plugins/Online/OnlineServicesEOSGS/Source/Private/OnlineServicesEOSGSModule.cpp
@@ -34,6 +34,13 @@ protected:
 
 void FOnlineServicesEOSGSModule::StartupModule()
 {
+	// Making sure we load the module at this point will avoid errors while cooking
+	const FName OnlineServicesInterfaceModuleName = TEXT("OnlineServicesInterface");
+	if (!FModuleManager::Get().IsModuleLoaded(OnlineServicesInterfaceModuleName))
+	{
+		FModuleManager::Get().LoadModuleChecked(OnlineServicesInterfaceModuleName);
+	}
+
 	FOnlineServicesRegistry::Get().RegisterServicesFactory(EOnlineServices::Epic, MakeUnique<FOnlineServicesFactoryEOSGS>());
 	FOnlineIdRegistryRegistry::Get().RegisterAccountIdRegistry(EOnlineServices::Epic, &FOnlineAccountIdRegistryEOSGS::Get());
 	FOnlineIdRegistryRegistry::Get().RegisterSessionIdRegistry(EOnlineServices::Epic, &FOnlineSessionIdRegistryEOSGS::Get());
diff --git a/Engine/Plugins/Online/OnlineServicesNull/Source/Private/Online/OnlineServicesNullModule.cpp b/Engine/Plugins/Online/OnlineServicesNull/Source/Private/Online/OnlineServicesNullModule.cpp
--- a/Engine/Plugins/Online/OnlineServicesNull/Source/Private/Online/OnlineServicesNullModule.cpp
+++ b/Engine/Plugins/Online/OnlineServicesNull/Source/Private/Online/OnlineServicesNullModule.cpp
@@ -33,6 +33,13 @@ protected:
 
 void FOnlineServicesNullModule::StartupModule()
 {
+	// Making sure we load the module at this point will avoid errors while cooking
+	const FName OnlineServicesInterfaceModuleName = TEXT("OnlineServicesInterface");
+	if (!FModuleManager::Get().IsModuleLoaded(OnlineServicesInterfaceModuleName))
+	{
+		FModuleManager::Get().LoadModuleChecked(OnlineServicesInterfaceModuleName);
+	}
+
 	FOnlineServicesRegistry::Get().RegisterServicesFactory(EOnlineServices::Null, MakeUnique<FOnlineServicesFactoryNull>());
 	FOnlineIdRegistryRegistry::Get().RegisterAccountIdRegistry(EOnlineServices::Null, &FOnlineAccountIdRegistryNull::Get());
 	FOnlineIdRegistryRegistry::Get().RegisterSessionIdRegistry(EOnlineServices::Null, &FOnlineSessionIdRegistryNull::Get());
