Refactored serializer logic to use bespoke FSchemaVariant method instead of treating all inputs as strings

diff --git a/Engine/Plugins/Online/OnlineServicesNull/Source/Public/Online/NboSerializerNullSvc.h b/Engine/Plugins/Online/OnlineServicesNull/Source/Public/Online/NboSerializerNullSvc.h
--- a/Engine/Plugins/Online/OnlineServicesNull/Source/Public/Online/NboSerializerNullSvc.h
+++ b/Engine/Plugins/Online/OnlineServicesNull/Source/Public/Online/NboSerializerNullSvc.h
@@ -49,7 +49,7 @@ inline void SerializeToBuffer(FNboSerializeToBuffer& Ar, const TMap<FSchemaAttri
 	for (const TPair<FSchemaAttributeId, FSchemaVariant>& Pair : Map)
 	{
 		Ar << Pair.Key;
-		Ar << Pair.Value.GetString();
+		NboSerializerCommonSvc::SerializeToBuffer(Ar, Pair.Value);
 	}
 }
 
@@ -96,9 +96,9 @@ inline void SerializeFromBuffer(FNboSerializeFromBuffer& Ar, TMap<FSchemaAttribu
 	for (int i = 0; i < Size; i++)
 	{
 		FSchemaAttributeId LobbyAttributeId;
-		FString LobbyData;
+		FSchemaVariant LobbyData;
 		Ar >> LobbyAttributeId;
-		Ar >> LobbyData;
+		NboSerializerCommonSvc::SerializeFromBuffer(Ar, LobbyData);
 		Map.Emplace(MoveTemp(LobbyAttributeId), MoveTemp(LobbyData));
 	}
 }
