- Move OnlineServicesNullModule class to public header
- Add "GetRegistryPriority" for ease of overriding.

diff --git a/Engine/Plugins/Online/OnlineServicesNull/Source/Private/Online/OnlineServicesNullModule.cpp b/Engine/Plugins/Online/OnlineServicesNull/Source/Private/Online/OnlineServicesNullModule.cpp
--- a/Engine/Plugins/Online/OnlineServicesNull/Source/Private/Online/OnlineServicesNullModule.cpp
+++ b/Engine/Plugins/Online/OnlineServicesNull/Source/Private/Online/OnlineServicesNullModule.cpp
@@ -1,7 +1,8 @@
 // Copyright Epic Games, Inc. All Rights Reserved.
 
-#include "Modules/ModuleManager.h"
+#include "Online/OnlineServicesNullModule.h"
 
+#include "Modules/ModuleManager.h"
 #include "Online/OnlineServicesRegistry.h"
 #include "Online/OnlineServicesNull.h"
 #include "Online/AuthNull.h"
@@ -10,14 +11,6 @@
 namespace UE::Online
 {
 
-class FOnlineServicesNullModule : public IModuleInterface
-{
-public:
-	virtual void StartupModule() override;
-	virtual void ShutdownModule() override;
-protected:
-};
-
 class FOnlineServicesFactoryNull : public IOnlineServicesFactory
 {
 public:
@@ -29,6 +22,11 @@ public:
 protected:
 };
 
+int FOnlineServicesNullModule::GetRegistryPriority()
+{
+	return 0;
+}
+
 void FOnlineServicesNullModule::StartupModule()
 {
 	// Making sure we load the module at this point will avoid errors while cooking
diff --git a/Engine/Plugins/Online/OnlineServicesNull/Source/Public/Online/AuthNull.h b/Engine/Plugins/Online/OnlineServicesNull/Source/Public/Online/AuthNull.h
--- a/Engine/Plugins/Online/OnlineServicesNull/Source/Public/Online/AuthNull.h
+++ b/Engine/Plugins/Online/OnlineServicesNull/Source/Public/Online/AuthNull.h
@@ -21,10 +21,10 @@ public:
 class FOnlineAccountIdRegistryNull : public IOnlineAccountIdRegistry
 {
 public:
-	static FOnlineAccountIdRegistryNull& Get();
+	ONLINESERVICESNULL_API static FOnlineAccountIdRegistryNull& Get();
 
-	FAccountId Find(const FString& AccountId) const;
-	FAccountId FindOrAddAccountId(const FString& AccountId);
+	ONLINESERVICESNULL_API FAccountId Find(const FString& AccountId) const;
+	ONLINESERVICESNULL_API FAccountId FindOrAddAccountId(const FString& AccountId);
 
 	// Begin IOnlineAccountIdRegistry
 	virtual FString ToString(const FAccountId& AccountId) const override;
diff --git a/Engine/Plugins/Online/OnlineServicesNull/Source/Public/Online/OnlineServicesNullModule.h b/Engine/Plugins/Online/OnlineServicesNull/Source/Public/Online/OnlineServicesNullModule.h
new file mode 100644
--- /dev/null
+++ b/Engine/Plugins/Online/OnlineServicesNull/Source/Public/Online/OnlineServicesNullModule.h
@@ -0,0 +1,21 @@
+// Copyright Epic Games, Inc. All Rights Reserved.
+
+#pragma once
+
+#include "Modules/ModuleInterface.h"
+
+namespace UE::Online {
+
+class FOnlineServicesNullModule : public IModuleInterface
+{
+public:
+	ONLINESERVICESNULL_API static int GetRegistryPriority();
+
+private:
+	// Begin IModuleInterface
+	virtual void StartupModule() override;
+	virtual void ShutdownModule() override;
+	// End IModuleInterface
+};
+
+} // namespace UE::Online
\ No newline at end of file
