Replace unnecessary Warnings with Verbose

diff --git a/Engine/Plugins/Online/OnlineServicesNull/Source/Private/Online/LobbiesNull.cpp b/Engine/Plugins/Online/OnlineServicesNull/Source/Private/Online/LobbiesNull.cpp
--- a/Engine/Plugins/Online/OnlineServicesNull/Source/Private/Online/LobbiesNull.cpp
+++ b/Engine/Plugins/Online/OnlineServicesNull/Source/Private/Online/LobbiesNull.cpp
@@ -107,7 +107,7 @@ TOnlineAsyncOpHandle<FCreateLobby> FLobbiesNull::CreateLobby(FCreateLobby::Param
 
 TOnlineAsyncOpHandle<FFindLobbies> FLobbiesNull::FindLobbies(FFindLobbies::Params&& Params)
 {
-	UE_LOG(LogOnlineServices, Warning, TEXT("FLobbiesNull::FindLobbies"));
+	UE_LOG(LogOnlineServices, Verbose, TEXT("FLobbiesNull::FindLobbies"));
 	if (!CurrentLobbySearch.IsValid())
 	{
 		CurrentLobbySearch = MakeShared<FFindLobbies::Result>();
@@ -211,7 +211,7 @@ TSharedRef<FLobbyNull> FLobbiesNull::CreateNamedLobby(const FCreateLobby::Params
 
 uint32 FLobbiesNull::FindLANSession()
 {
-	UE_LOG(LogOnlineServices, Warning, TEXT("FLobbiesNull::FindLANSession"));
+	UE_LOG(LogOnlineServices, Verbose, TEXT("FLobbiesNull::FindLANSession"));
 	uint32 Return = ONLINE_IO_PENDING;
 
 	// Recreate the unique identifier for this client
@@ -236,7 +236,7 @@ uint32 FLobbiesNull::FindLANSession()
 	}
 	else
 	{
-		UE_LOG(LogOnlineServices, Warning, TEXT("FLobbiesNull::FindLANSession: searching...."));
+		UE_LOG(LogOnlineServices, Verbose, TEXT("FLobbiesNull::FindLANSession: searching...."));
 	}
 
 
@@ -251,7 +251,7 @@ bool FLobbiesNull::NeedsToAdvertise()
 
 uint32 FLobbiesNull::UpdateLANStatus()
 {
-	UE_LOG(LogOnlineServices, Warning, TEXT("FLobbiesNull::UpdateLANStatus"));
+	UE_LOG(LogOnlineServices, Verbose, TEXT("FLobbiesNull::UpdateLANStatus"));
 	uint32 Result = ONLINE_SUCCESS;
 
 	if ( NeedsToAdvertise() )
@@ -259,7 +259,7 @@ uint32 FLobbiesNull::UpdateLANStatus()
 		// set up LAN session
 		if (LANSessionManager.GetBeaconState() == ELanBeaconState::NotUsingLanBeacon)
 		{
-			UE_LOG(LogOnlineServices, Warning, TEXT("FLobbiesNull::UpdateLANStatus- Hosting.."));
+			UE_LOG(LogOnlineServices, Verbose, TEXT("FLobbiesNull::UpdateLANStatus- Hosting.."));
 			FOnValidQueryPacketDelegate QueryPacketDelegate = FOnValidQueryPacketDelegate::CreateRaw(this, &FLobbiesNull::OnValidQueryPacketReceived);
 			if (!LANSessionManager.Host(QueryPacketDelegate))
 			{
@@ -287,7 +287,7 @@ void FLobbiesNull::AppendLobbyToPacket(FNboSerializeToBuffer& Packet, const TSha
 	using namespace NboSerializerNullSvc;
 
 	/** Owner of the session */
-	UE_LOG(LogOnlineServices, Warning, TEXT("AppendLobbyToPacket: Appending IP %s"), *Lobby->HostAddr->ToString(true));
+	UE_LOG(LogOnlineServices, Verbose, TEXT("AppendLobbyToPacket: Appending IP %s"), *Lobby->HostAddr->ToString(true));
 
 	SerializeToBuffer(Packet, Lobby->Data->OwnerAccountId);
 	SerializeToBuffer(Packet, Lobby->Data->Attributes);
@@ -302,7 +302,7 @@ void FLobbiesNull::AppendLobbyToPacket(FNboSerializeToBuffer& Packet, const TSha
 
 void FLobbiesNull::OnValidQueryPacketReceived(uint8* PacketData, int32 PacketLength, uint64 ClientNonce)
 {
-	UE_LOG(LogOnlineServices, Warning, TEXT("FLobbiesNull::OnValidQueryPacketReceived"));
+	UE_LOG(LogOnlineServices, Verbose, TEXT("FLobbiesNull::OnValidQueryPacketReceived"));
 	// Iterate through all registered sessions and respond for each one that can be joinable
 	for (TPair<FName, TSharedRef<FLobbyNull>>& Pair : NamedLobbies)
 	{
@@ -353,7 +353,7 @@ void FLobbiesNull::ReadLobbyFromPacket(FNboSerializeFromBuffer& Packet, const TS
 	Session->HostAddr->SetPort(Session->HostAddrPort);
 	Session->Data->Attributes.Add(FName(TEXT("ConnectAddress")), Session->HostAddr->ToString(true));
 	Session->Data->LobbyId = FOnlineLobbyIdRegistryNull::Get().GetNext(); //TODO: this should be coming from the packet, not a random GUID
-	UE_LOG(LogOnlineServices, Warning, TEXT("ReadLobby: Got IP %s"), *Session->HostAddr->ToString(true));
+	UE_LOG(LogOnlineServices, Verbose, TEXT("ReadLobby: Got IP %s"), *Session->HostAddr->ToString(true));
 
 	// todo: make all the members (not just the host
 	TSharedRef<FLobbyMember> NewMember = MakeShared<FLobbyMember>();
@@ -366,7 +366,7 @@ void FLobbiesNull::ReadLobbyFromPacket(FNboSerializeFromBuffer& Packet, const TS
 
 void FLobbiesNull::OnValidResponsePacketReceived(uint8* PacketData, int32 PacketLength)
 {
-	UE_LOG(LogOnlineServices, Warning, TEXT("FLobbiesNull::OnValidResponsePacketReceived"));
+	UE_LOG(LogOnlineServices, Verbose, TEXT("FLobbiesNull::OnValidResponsePacketReceived"));
 	if (CurrentLobbySearch.IsValid())
 	{
 		TSharedRef<FLobbyNull> Lobby = MakeShared<FLobbyNull>();
@@ -392,7 +392,7 @@ void FLobbiesNull::OnValidResponsePacketReceived(uint8* PacketData, int32 Packet
 
 uint32 FLobbiesNull::FinalizeLANSearch()
 {
-	UE_LOG(LogOnlineServices, Warning, TEXT("FLobbiesNull::FinalizeLANSearch"));
+	UE_LOG(LogOnlineServices, Verbose, TEXT("FLobbiesNull::FinalizeLANSearch"));
 	if (LANSessionManager.GetBeaconState() == ELanBeaconState::Searching)
 	{
 		LANSessionManager.StopLANSession();
