Fix the online tests failure where auth null implementation can return invalid account id due to referencing invalid memory chunk.

diff --git a/Engine/Plugins/Online/OnlineServicesNull/Source/Private/Online/AuthNull.cpp b/Engine/Plugins/Online/OnlineServicesNull/Source/Private/Online/AuthNull.cpp
--- a/Engine/Plugins/Online/OnlineServicesNull/Source/Private/Online/AuthNull.cpp
+++ b/Engine/Plugins/Online/OnlineServicesNull/Source/Private/Online/AuthNull.cpp
@@ -194,8 +194,13 @@ const FOnlineAccountIdString* FOnlineAccountIdRegistryNull::FindNoLock(const FAc
 
 const FOnlineAccountIdString* FOnlineAccountIdRegistryNull::FindNoLock(const FString& AccountId) const
 {
-	FOnlineAccountIdString* const* Entry = StringToIdIndex.Find(AccountId);
-	return Entry ? *Entry : nullptr;
+	const uint32* Index = StringToIdIndex.Find(AccountId);
+	if (!Index)
+	{
+		return nullptr;
+	}
+
+	return &Ids[*Index];
 }
 
 FAccountId FOnlineAccountIdRegistryNull::FindOrAddAccountId(const FString& AccountId)
@@ -222,7 +227,7 @@ FAccountId FOnlineAccountIdRegistryNull::FindOrAddAccountId(const FString& Accou
 		Id.Data = AccountId;
 		Id.AccountId = FAccountId(EOnlineServices::Null, Id.AccountIndex);
 
-		StringToIdIndex.Add(AccountId, &Id);
+		StringToIdIndex.Add(AccountId, Ids.Num() - 1);
 		return Id.AccountId;
 	}
 }
diff --git a/Engine/Plugins/Online/OnlineServicesNull/Source/Public/Online/AuthNull.h b/Engine/Plugins/Online/OnlineServicesNull/Source/Public/Online/AuthNull.h
--- a/Engine/Plugins/Online/OnlineServicesNull/Source/Public/Online/AuthNull.h
+++ b/Engine/Plugins/Online/OnlineServicesNull/Source/Public/Online/AuthNull.h
@@ -42,7 +42,7 @@ private:
 
 	mutable FRWLock Lock;
 	TArray<FOnlineAccountIdString> Ids;
-	TMap<FString, FOnlineAccountIdString*> StringToIdIndex;
+	TMap<FString, uint32> StringToIdIndex;
 
 };
 
