Make FOnlineAccountIdRegistryNull use the common account id registry logic provided by TOnlineBasicAccountIdRegistry

diff --git a/Engine/Plugins/Online/OnlineServicesNull/Source/Private/Online/AuthNull.cpp b/Engine/Plugins/Online/OnlineServicesNull/Source/Private/Online/AuthNull.cpp
--- a/Engine/Plugins/Online/OnlineServicesNull/Source/Private/Online/AuthNull.cpp
+++ b/Engine/Plugins/Online/OnlineServicesNull/Source/Private/Online/AuthNull.cpp
@@ -170,77 +170,41 @@ void FAuthNull::OnInputDeviceConnectionChange(EInputDeviceConnectionState NewCon
 }
 
 // FOnlineAccountIdRegistryNull
-FOnlineAccountIdRegistryNull& FOnlineAccountIdRegistryNull::Get()
-{
-	static FOnlineAccountIdRegistryNull Instance;
-	return Instance;
-}
 
-FAccountId FOnlineAccountIdRegistryNull::Find(const FString& AccountId) const
+FOnlineAccountIdRegistryNull::FOnlineAccountIdRegistryNull()
+: Registry(EOnlineServices::Null)
 {
-	const FReadScopeLock ReadLock(Lock);
-	const FOnlineAccountIdString* Entry = FindNoLock(AccountId);
-	return Entry ? Entry->AccountId : FAccountId();
 }
 
-const FOnlineAccountIdString* FOnlineAccountIdRegistryNull::FindNoLock(const FAccountId& AccountId) const
+FOnlineAccountIdRegistryNull& FOnlineAccountIdRegistryNull::Get()
 {
-	if(AccountId.IsValid() && AccountId.GetOnlineServicesType() == EOnlineServices::Null && AccountId.GetHandle() <= (uint32)Ids.Num())
-	{
-		return &Ids[AccountId.GetHandle()-1];
-	}
-	return nullptr;
+	static FOnlineAccountIdRegistryNull Instance;
+	return Instance;
 }
 
-const FOnlineAccountIdString* FOnlineAccountIdRegistryNull::FindNoLock(const FString& AccountId) const
+FAccountId FOnlineAccountIdRegistryNull::Find(const FString& AccountId) const
 {
-	const uint32* Index = StringToIdIndex.Find(AccountId);
-	if (!Index)
-	{
-		return nullptr;
-	}
-
-	return &Ids[*Index];
+	return Registry.FindHandle(AccountId);
 }
 
 FAccountId FOnlineAccountIdRegistryNull::FindOrAddAccountId(const FString& AccountId)
 {
-	// Check for existing entry under read lock.
-	{
-		const FReadScopeLock ReadLock(Lock);
-		if (const FOnlineAccountIdString* ExistingAccountId = FindNoLock(AccountId))
-		{
-			return ExistingAccountId->AccountId;
-		}
-	}
-
-	// Check for existing entry again under write lock before adding entry.
-	{
-		const FWriteScopeLock WriteLock(Lock);
-		if (const FOnlineAccountIdString* ExistingAccountId = FindNoLock(AccountId))
-		{
-			return ExistingAccountId->AccountId;
-		}
-
-		FOnlineAccountIdString& Id = Ids.Emplace_GetRef();
-		Id.AccountIndex = Ids.Num();
-		Id.Data = AccountId;
-		Id.AccountId = FAccountId(EOnlineServices::Null, Id.AccountIndex);
-
-		StringToIdIndex.Add(AccountId, Ids.Num() - 1);
-		return Id.AccountId;
-	}
+	return Registry.FindOrAddHandle(AccountId);
 }
 
 FString FOnlineAccountIdRegistryNull::ToString(const FAccountId& AccountId) const
 {
-	const FReadScopeLock ReadLock(Lock);
-	if(const FOnlineAccountIdString* Id = FindNoLock(AccountId))
+	FString Result;
+	if (Registry.ValidateOnlineId(AccountId))
 	{
-		return Id->Data;
+		Result = Registry.FindIdValue(AccountId);
 	}
-
-	return FString(TEXT("Invalid"));
+	else
+	{
+		check(!AccountId.IsValid()); // Check we haven't been passed a valid handle for a different EOnlineServices.
+		Result = TEXT("Invalid");
+	}
+	return Result;
 }
 
 FString FOnlineAccountIdRegistryNull::ToLogString(const FAccountId& AccountId) const
@@ -250,17 +214,15 @@ FString FOnlineAccountIdRegistryNull::ToLogString(const FAccountId& AccountId) c
 
 TArray<uint8> FOnlineAccountIdRegistryNull::ToReplicationData(const FAccountId& AccountId) const
 {
-	const FReadScopeLock ReadLock(Lock);
-	if (const FOnlineAccountIdString* Id = FindNoLock(AccountId))
+	TArray<uint8> ReplicationData;
+	if (Registry.ValidateOnlineId(AccountId))
 	{
-		TArray<uint8> ReplicationData;
-		ReplicationData.SetNumUninitialized(Id->Data.Len());
-		StringToBytes(Id->Data, ReplicationData.GetData(), Id->Data.Len());
-		UE_LOG(LogOnlineServices, VeryVerbose, TEXT("[FOnlineAccountIdRegistryNull::ToReplicationData] StringToBytes on %s returned %d len"), *Id->Data, ReplicationData.Num())
-		return ReplicationData;
+		const FString& AccountIdString = Registry.FindIdValue(AccountId);
+		ReplicationData.SetNumUninitialized(AccountIdString.Len());
+		StringToBytes(AccountIdString, ReplicationData.GetData(), ReplicationData.Num());
+		UE_LOG(LogOnlineServices, VeryVerbose, TEXT("[FOnlineAccountIdRegistryNull::ToReplicationData] StringToBytes on [%s] returned %d len"), *AccountIdString, ReplicationData.Num())
 	}
-
-	return TArray<uint8>();
+	return ReplicationData;
 }
 
 FAccountId FOnlineAccountIdRegistryNull::FromReplicationData(const TArray<uint8>& ReplicationData)
diff --git a/Engine/Plugins/Online/OnlineServicesNull/Source/Public/Online/AuthNull.h b/Engine/Plugins/Online/OnlineServicesNull/Source/Public/Online/AuthNull.h
--- a/Engine/Plugins/Online/OnlineServicesNull/Source/Public/Online/AuthNull.h
+++ b/Engine/Plugins/Online/OnlineServicesNull/Source/Public/Online/AuthNull.h
@@ -5,12 +5,13 @@
 #include "CoreMinimal.h"
 #include "HAL/CriticalSection.h"
 #include "Online/AuthCommon.h"
+#include "Online/OnlineIdCommon.h"
 
 namespace UE::Online {
 
 class FOnlineServicesNull;
 
-class FOnlineAccountIdString
+class UE_DEPRECATED(5.5, "FOnlineAccountIdRegistryNull now uses TOnlineBasicAccountIdRegistry, and FOnlineAccountIdString was never exposed external to FOnlineAccountIdRegistryNull") FOnlineAccountIdString
 {
 public:
 	FString Data;
@@ -18,7 +19,8 @@ public:
 	FAccountId AccountId;
 };
 
-class FOnlineAccountIdRegistryNull : public IOnlineAccountIdRegistry
+class FOnlineAccountIdRegistryNull
+	: public IOnlineAccountIdRegistry
 {
 public:
 	ONLINESERVICESNULL_API static FOnlineAccountIdRegistryNull& Get();
@@ -36,21 +38,15 @@ public:
 	virtual ~FOnlineAccountIdRegistryNull() = default;
 
 private:
-	// Retrieve an entry with no lock. For internal use only.
-	const FOnlineAccountIdString* FindNoLock(const FAccountId& AccountId) const;
-	const FOnlineAccountIdString* FindNoLock(const FString& AccountId) const;
-
-	mutable FRWLock Lock;
-	TArray<FOnlineAccountIdString> Ids;
-	TMap<FString, uint32> StringToIdIndex;
-
+	FOnlineAccountIdRegistryNull();
+	TOnlineBasicAccountIdRegistry<FString> Registry;
 };
 
+
 struct FAccountInfoNull final : public FAccountInfo
 {
 };
 
-
 // Auth NULL is implemented in a way similar to console platforms where there is not an explicit
 // login / logout from online services. On those platforms the user account is picked either before
 // the game has started or as a part of selecting an input device.
