Make OnlineServicesNull cache the config instance name so users can distinguish between null instances

diff --git a/Engine/Plugins/Online/OnlineServicesNull/Source/Private/Online/OnlineServicesNull.cpp b/Engine/Plugins/Online/OnlineServicesNull/Source/Private/Online/OnlineServicesNull.cpp
--- a/Engine/Plugins/Online/OnlineServicesNull/Source/Private/Online/OnlineServicesNull.cpp
+++ b/Engine/Plugins/Online/OnlineServicesNull/Source/Private/Online/OnlineServicesNull.cpp
@@ -29,8 +29,8 @@ END_ONLINE_STRUCT_META()
 
  /*Meta*/ }
 
-FOnlineServicesNull::FOnlineServicesNull(FName InInstanceName)
-	: FOnlineServicesCommon(TEXT("Null"), InInstanceName)
+FOnlineServicesNull::FOnlineServicesNull(FName InInstanceName, FName InInstanceConfigName)
+	: FOnlineServicesCommon(TEXT("Null"), InInstanceName, InInstanceConfigName)
 {
 }
 
diff --git a/Engine/Plugins/Online/OnlineServicesNull/Source/Private/Online/OnlineServicesNullModule.cpp b/Engine/Plugins/Online/OnlineServicesNull/Source/Private/Online/OnlineServicesNullModule.cpp
--- a/Engine/Plugins/Online/OnlineServicesNull/Source/Private/Online/OnlineServicesNullModule.cpp
+++ b/Engine/Plugins/Online/OnlineServicesNull/Source/Private/Online/OnlineServicesNullModule.cpp
@@ -17,7 +17,7 @@ public:
 	virtual ~FOnlineServicesFactoryNull() {}
 	virtual TSharedPtr<IOnlineServices> Create(FName InInstanceName, FName InstanceConfigName) override
 	{
-		return MakeShared<FOnlineServicesNull>(InInstanceName);
+		return MakeShared<FOnlineServicesNull>(InInstanceName, InstanceConfigName);
 	}
 protected:
 };
diff --git a/Engine/Plugins/Online/OnlineServicesNull/Source/Public/Online/OnlineServicesNull.h b/Engine/Plugins/Online/OnlineServicesNull/Source/Public/Online/OnlineServicesNull.h
--- a/Engine/Plugins/Online/OnlineServicesNull/Source/Public/Online/OnlineServicesNull.h
+++ b/Engine/Plugins/Online/OnlineServicesNull/Source/Public/Online/OnlineServicesNull.h
@@ -11,7 +11,7 @@ class ONLINESERVICESNULL_API FOnlineServicesNull : public FOnlineServicesCommon
 public:
 	using Super = FOnlineServicesCommon;
 
-	FOnlineServicesNull(FName InInstanceName);
+	FOnlineServicesNull(FName InInstanceName, FName InInstanceConfigName);
 	virtual void RegisterComponents() override;
 	virtual void Initialize() override;
 	virtual TOnlineResult<FGetResolvedConnectString> GetResolvedConnectString(FGetResolvedConnectString::Params&& Params) override;
