Fix the check caused by invalid platform user id in Auth impl when InputDeviceConnectionChange got triggered.

diff --git a/Engine/Plugins/Online/OnlineServicesNull/Source/Private/Online/AuthNull.cpp b/Engine/Plugins/Online/OnlineServicesNull/Source/Private/Online/AuthNull.cpp
--- a/Engine/Plugins/Online/OnlineServicesNull/Source/Private/Online/AuthNull.cpp
+++ b/Engine/Plugins/Online/OnlineServicesNull/Source/Private/Online/AuthNull.cpp
@@ -158,7 +158,7 @@ void FAuthNull::UninitializeUsers()
 void FAuthNull::OnInputDeviceConnectionChange(EInputDeviceConnectionState NewConnectionState, FPlatformUserId PlatformUserId, FInputDeviceId InputDeviceId)
 {
 	// If this is a new platform user then register an entry for them so they will be seen as "logged-in".
-	if (!AccountInfoRegistryNULL.Find(PlatformUserId))
+	if (NewConnectionState == EInputDeviceConnectionState::Connected && PlatformUserId != PLATFORMUSERID_NONE && !AccountInfoRegistryNULL.Find(PlatformUserId))
 	{
 		FAuthNullConfig AuthNullConfig;
 		LoadConfig(AuthNullConfig);
