Fix OSS Null Presence updates

diff --git a/Engine/Plugins/Online/OnlineServices/Source/OnlineServicesInterface/Public/Online/Presence.h b/Engine/Plugins/Online/OnlineServices/Source/OnlineServicesInterface/Public/Online/Presence.h
--- a/Engine/Plugins/Online/OnlineServices/Source/OnlineServicesInterface/Public/Online/Presence.h
+++ b/Engine/Plugins/Online/OnlineServices/Source/OnlineServicesInterface/Public/Online/Presence.h
@@ -279,6 +279,11 @@ inline TSharedRef<FUserPresence> ApplyPresenceMutations(const FUserPresence& Bas
 		NewPresence->StatusString = Mutations.StatusString.GetValue();
 	}
 
+	if (Mutations.RichPresenceString.IsSet())
+	{
+		NewPresence->RichPresenceString = Mutations.RichPresenceString.GetValue();
+	}
+
 	// Insert any updated keys / remove from removed
 	for (const TPair<FString, FPresenceVariant>& NewUpdatedProperty : Mutations.UpdatedProperties)
 	{
diff --git a/Engine/Plugins/Online/OnlineServicesNull/Source/Private/Online/PresenceNull.cpp b/Engine/Plugins/Online/OnlineServicesNull/Source/Private/Online/PresenceNull.cpp
--- a/Engine/Plugins/Online/OnlineServicesNull/Source/Private/Online/PresenceNull.cpp
+++ b/Engine/Plugins/Online/OnlineServicesNull/Source/Private/Online/PresenceNull.cpp
@@ -96,9 +96,13 @@ TOnlineAsyncOpHandle<FPartialUpdatePresence> FPresenceNull::PartialUpdatePresenc
 		TSharedRef<FUserPresence> NewPresence = MakeShared<FUserPresence>();
 		const FPartialUpdatePresence::Params::FMutations& Mutations = InAsyncOp.GetParams().Mutations;
 		
-		if (Presences.Contains(InAsyncOp.GetParams().LocalAccountId))
+		if (const TSharedRef<const FUserPresence>* const ExistingPresence = Presences.Find(InAsyncOp.GetParams().LocalAccountId))
 		{
-			*NewPresence = *Presences.FindChecked(InAsyncOp.GetParams().LocalAccountId);
+			*NewPresence = **ExistingPresence;
+		}
+		else
+		{
+			NewPresence->AccountId = InAsyncOp.GetParams().LocalAccountId;
 		}
 
 		TSharedRef<FUserPresence> MutatedPresence = ApplyPresenceMutations(*NewPresence, Mutations);
