Export FOnlineLobbyIdRegistryNull

#Synced-CL 51004239

#Synced-CL 51004239

diff --git a/Engine/Plugins/Online/OnlineServicesNull/Source/Public/Online/LobbyRegistryNull.h b/Engine/Plugins/Online/OnlineServicesNull/Source/Public/Online/LobbyRegistryNull.h
--- a/Engine/Plugins/Online/OnlineServicesNull/Source/Public/Online/LobbyRegistryNull.h
+++ b/Engine/Plugins/Online/OnlineServicesNull/Source/Public/Online/LobbyRegistryNull.h
@@ -10,11 +10,11 @@ using IOnlineLobbyIdRegistry = IOnlineIdRegistry<OnlineIdHandleTags::FLobby>; //
 class FOnlineLobbyIdRegistryNull : public IOnlineLobbyIdRegistry
 {
 public:
-	static FOnlineLobbyIdRegistryNull& Get();
+	ONLINESERVICESNULL_API static FOnlineLobbyIdRegistryNull& Get();
 
-	const FLobbyId* Find(FString LobbyIdStr);
-	FLobbyId FindOrAdd(FString LobbyIdStr);
-	FLobbyId GetNext();
+	ONLINESERVICESNULL_API const FLobbyId* Find(FString LobbyIdStr);
+	ONLINESERVICESNULL_API FLobbyId FindOrAdd(FString LobbyIdStr);
+	ONLINESERVICESNULL_API FLobbyId GetNext();
 
 	// Begin IOnlineAccountIdRegistry
 	virtual FString ToString(const FLobbyId& LobbyId) const override;
