Compile fixes for programs that do not compile against Application Core and/or Engine.

#Synced-CL 51345931

diff --git a/Engine/Plugins/Online/OnlineFrameworkCommon/Source/OnlineFrameworkCommon/OnlineFrameworkCommon.Build.cs b/Engine/Plugins/Online/OnlineFrameworkCommon/Source/OnlineFrameworkCommon/OnlineFrameworkCommon.Build.cs
--- a/Engine/Plugins/Online/OnlineFrameworkCommon/Source/OnlineFrameworkCommon/OnlineFrameworkCommon.Build.cs
+++ b/Engine/Plugins/Online/OnlineFrameworkCommon/Source/OnlineFrameworkCommon/OnlineFrameworkCommon.Build.cs
@@ -14,9 +14,13 @@ public class OnlineFrameworkCommon : ModuleRules
 				"Core",
 				"CoreOnline",
 				"CoreUObject",
-				"Engine",
 			});
 
+		if (Target.bCompileAgainstEngine)
+		{
+			PublicDependencyModuleNames.Add("Engine");
+		}
+
 		PrivateDependencyModuleNames.AddRange(
 			new string[]
 			{
diff --git a/Engine/Plugins/Online/OnlineFrameworkCommon/Source/OnlineFrameworkCommon/Private/OnlineFramework/CommonConfig.cpp b/Engine/Plugins/Online/OnlineFrameworkCommon/Source/OnlineFrameworkCommon/Private/OnlineFramework/CommonConfig.cpp
--- a/Engine/Plugins/Online/OnlineFrameworkCommon/Source/OnlineFrameworkCommon/Private/OnlineFramework/CommonConfig.cpp
+++ b/Engine/Plugins/Online/OnlineFrameworkCommon/Source/OnlineFrameworkCommon/Private/OnlineFramework/CommonConfig.cpp
@@ -2,8 +2,11 @@
 
 #include "OnlineFramework/CommonConfig.h"
 
+#if WITH_ENGINE
 #include "Engine/Engine.h"
 #include "Engine/World.h"
+#endif
+
 #include "Online/OnlineServices.h"
 #include "OnlineFramework/CommonModule.h"
 #include "Templates/SharedPointer.h"
diff --git a/Engine/Plugins/Online/OnlineServicesNull/Source/OnlineServicesNull.Build.cs b/Engine/Plugins/Online/OnlineServicesNull/Source/OnlineServicesNull.Build.cs
--- a/Engine/Plugins/Online/OnlineServicesNull/Source/OnlineServicesNull.Build.cs
+++ b/Engine/Plugins/Online/OnlineServicesNull/Source/OnlineServicesNull.Build.cs
@@ -20,10 +20,14 @@ public class OnlineServicesNull : ModuleRules
 
 		PrivateDependencyModuleNames.AddRange(
 			new string[] { 
-				"ApplicationCore",
 				"OnlineBase",
 				"OnlineServicesCommonEngineUtils"
 			}
 		);
+
+		if (Target.bCompileAgainstApplicationCore)
+		{
+			PrivateDependencyModuleNames.Add("ApplicationCore");
+		}
 	}
 }
diff --git a/Engine/Plugins/Online/OnlineServicesNull/Source/Private/Online/AuthNull.cpp b/Engine/Plugins/Online/OnlineServicesNull/Source/Private/Online/AuthNull.cpp
--- a/Engine/Plugins/Online/OnlineServicesNull/Source/Private/Online/AuthNull.cpp
+++ b/Engine/Plugins/Online/OnlineServicesNull/Source/Private/Online/AuthNull.cpp
@@ -9,7 +9,9 @@
 #include "Online/OnlineServicesNull.h"
 #include "SocketSubsystem.h"
 
+#if WITH_APPLICATION_CORE
 #include "GenericPlatform/GenericPlatformInputDeviceMapper.h"
+#endif
 
 namespace UE::Online {
 // Copied from OSS Null
@@ -138,6 +140,7 @@ void FAuthNull::InitializeUsers()
 	FAuthNullConfig AuthNullConfig;
 	LoadConfig(AuthNullConfig);
 
+#if WITH_APPLICATION_CORE // for IPlatformInputDeviceMapper
 	// There is no "login" for Null - all local users are initialized as "logged in".
 	TArray<FPlatformUserId> Users;
 	IPlatformInputDeviceMapper::Get().GetAllActiveUsers(Users);
@@ -148,11 +151,14 @@ void FAuthNull::InitializeUsers()
 
 	// Setup hook to add new users when they become available.
 	IPlatformInputDeviceMapper::Get().GetOnInputDeviceConnectionChange().AddRaw(this, &FAuthNull::OnInputDeviceConnectionChange);
+#endif
 }
 
 void FAuthNull::UninitializeUsers()
 {
+#if WITH_APPLICATION_CORE // for IPlatformInputDeviceMapper
 	IPlatformInputDeviceMapper::Get().GetOnInputDeviceConnectionChange().RemoveAll(this);
+#endif
 }
 
 void FAuthNull::OnInputDeviceConnectionChange(EInputDeviceConnectionState NewConnectionState, FPlatformUserId PlatformUserId, FInputDeviceId InputDeviceId)
